Move MeshBuilder module from runtime to developper and make it editor only

diff --git a/Engine/Source/Developer/MeshBuilder/MeshBuilder.Build.cs b/Engine/Source/Developer/MeshBuilder/MeshBuilder.Build.cs
new file mode 100644
--- /dev/null
+++ b/Engine/Source/Developer/MeshBuilder/MeshBuilder.Build.cs
@@ -0,0 +1,48 @@
+// Copyright 1998-2017 Epic Games, Inc. All Rights Reserved.
+
+namespace UnrealBuildTool.Rules
+{
+    //MeshBuilder module is a editor module
+	public class MeshBuilder : ModuleRules
+	{
+		public MeshBuilder(ReadOnlyTargetRules Target) : base(Target)
+		{
+            PrivateIncludePaths.Add("Runtime/MeshBuilder/Private");
+            PublicIncludePaths.Add("Runtime/MeshBuilder/Public");
+
+            PrivateDependencyModuleNames.AddRange(
+				new string[]
+				{
+					"Core",
+					"CoreUObject",
+                    "Engine",
+                    "RenderCore",
+                    "MeshDescription",
+                    "RenderCore",
+                    "MeshDescriptionOperations",
+                    "MeshReductionInterface",
+                    "RawMesh"
+                }
+			);
+
+			AddEngineThirdPartyPrivateStaticDependencies(Target, "nvTriStrip");
+            AddEngineThirdPartyPrivateStaticDependencies(Target, "ForsythTriOptimizer");
+	        AddEngineThirdPartyPrivateStaticDependencies(Target, "nvTessLib");
+            AddEngineThirdPartyPrivateStaticDependencies(Target, "QuadricMeshReduction");
+
+            if (Target.bCompileSimplygon == true)
+            {
+                AddEngineThirdPartyPrivateDynamicDependencies(Target, "SimplygonMeshReduction");
+
+                if (Target.bCompileSimplygonSSF == true)
+                {
+                    DynamicallyLoadedModuleNames.AddRange(
+                        new string[] {
+                    "SimplygonSwarm"
+                    }
+                    );
+                }
+            }
+        }
+	}
+}
diff --git a/Engine/Source/Developer/MeshBuilder/Private/BuildOptimizationHelper.h b/Engine/Source/Developer/MeshBuilder/Private/BuildOptimizationHelper.h
new file mode 100644
--- /dev/null
+++ b/Engine/Source/Developer/MeshBuilder/Private/BuildOptimizationHelper.h
@@ -0,0 +1,364 @@
+// Copyright 1998-2017 Epic Games, Inc. All Rights Reserved.
+#pragma once
+
+#include "CoreMinimal.h"
+#include "CoreTypes.h"
+#include "Components.h"
+#include "StaticMeshResources.h"
+#include "ThirdPartyBuildOptimizationHelper.h"
+
+class BuildOptimizationHelper
+{
+public:
+	/** Helper struct for building acceleration structures. */
+	struct FIndexAndZ
+	{
+		float Z;
+		int32 Index;
+
+		/** Default constructor. */
+		FIndexAndZ() {}
+
+		/** Initialization constructor. */
+		FIndexAndZ(int32 InIndex, FVector V)
+		{
+			Z = 0.30f * V.X + 0.33f * V.Y + 0.37f * V.Z;
+			Index = InIndex;
+		}
+	};
+
+	/** Sorting function for vertex Z/index pairs. */
+	struct FCompareIndexAndZ
+	{
+		FORCEINLINE bool operator()(FIndexAndZ const& A, FIndexAndZ const& B) const { return A.Z < B.Z; }
+	};
+
+	struct FMeshEdge
+	{
+		int32	Vertices[2];
+		int32	Faces[2];
+	};
+
+	/**
+	* This helper class builds the edge list for a mesh. It uses a hash of vertex
+	* positions to edges sharing that vertex to remove the n^2 searching of all
+	* previously added edges. This class is templatized so it can be used with
+	* either static mesh or skeletal mesh vertices
+	*/
+	template <class VertexClass> class TEdgeBuilder
+	{
+	protected:
+		/**
+		* The list of indices to build the edge data from
+		*/
+		const TArray<uint32>& Indices;
+		/**
+		* The array of verts for vertex position comparison
+		*/
+		const TArray<VertexClass>& Vertices;
+		/**
+		* The array of edges to create
+		*/
+		TArray<FMeshEdge>& Edges;
+		/**
+		* List of edges that start with a given vertex
+		*/
+		TMultiMap<FVector, FMeshEdge*> VertexToEdgeList;
+
+		/**
+		* This function determines whether a given edge matches or not. It must be
+		* provided by derived classes since they have the specific information that
+		* this class doesn't know about (vertex info, influences, etc)
+		*
+		* @param Index1 The first index of the edge being checked
+		* @param Index2 The second index of the edge
+		* @param OtherEdge The edge to compare. Was found via the map
+		*
+		* @return true if the edge is a match, false otherwise
+		*/
+		virtual bool DoesEdgeMatch(int32 Index1, int32 Index2, FMeshEdge* OtherEdge) = 0;
+
+		/**
+		* Searches the list of edges to see if this one matches an existing and
+		* returns a pointer to it if it does
+		*
+		* @param Index1 the first index to check for
+		* @param Index2 the second index to check for
+		*
+		* @return NULL if no edge was found, otherwise the edge that was found
+		*/
+		inline FMeshEdge* FindOppositeEdge(int32 Index1, int32 Index2)
+		{
+			FMeshEdge* Edge = NULL;
+			TArray<FMeshEdge*> EdgeList;
+			// Search the hash for a corresponding vertex
+			VertexToEdgeList.MultiFind(Vertices[Index2].Position, EdgeList);
+			// Now search through the array for a match or not
+			for (int32 EdgeIndex = 0; EdgeIndex < EdgeList.Num() && Edge == NULL;
+				EdgeIndex++)
+			{
+				FMeshEdge* OtherEdge = EdgeList[EdgeIndex];
+				// See if this edge matches the passed in edge
+				if (OtherEdge != NULL && DoesEdgeMatch(Index1, Index2, OtherEdge))
+				{
+					// We have a match
+					Edge = OtherEdge;
+				}
+			}
+			return Edge;
+		}
+
+		/**
+		* Updates an existing edge if found or adds the new edge to the list
+		*
+		* @param Index1 the first index in the edge
+		* @param Index2 the second index in the edge
+		* @param Triangle the triangle that this edge was found in
+		*/
+		inline void AddEdge(int32 Index1, int32 Index2, int32 Triangle)
+		{
+			// If this edge matches another then just fill the other triangle
+			// otherwise add it
+			FMeshEdge* OtherEdge = FindOppositeEdge(Index1, Index2);
+			if (OtherEdge == NULL)
+			{
+				// Add a new edge to the array
+				int32 EdgeIndex = Edges.AddZeroed();
+				Edges[EdgeIndex].Vertices[0] = Index1;
+				Edges[EdgeIndex].Vertices[1] = Index2;
+				Edges[EdgeIndex].Faces[0] = Triangle;
+				Edges[EdgeIndex].Faces[1] = -1;
+				// Also add this edge to the hash for faster searches
+				// NOTE: This relies on the array never being realloced!
+				VertexToEdgeList.Add(Vertices[Index1].Position, &Edges[EdgeIndex]);
+			}
+			else
+			{
+				OtherEdge->Faces[1] = Triangle;
+			}
+		}
+
+	public:
+		/**
+		* Initializes the values for the code that will build the mesh edge list
+		*/
+		TEdgeBuilder(const TArray<uint32>& InIndices,
+			const TArray<VertexClass>& InVertices,
+			TArray<FMeshEdge>& OutEdges) :
+			Indices(InIndices), Vertices(InVertices), Edges(OutEdges)
+		{
+			// Presize the array so that there are no extra copies being done
+			// when adding edges to it
+			Edges.Empty(Indices.Num());
+		}
+
+		/**
+		* Virtual dtor
+		*/
+		virtual ~TEdgeBuilder() {}
+
+
+		/**
+		* Uses a hash of indices to edge lists so that it can avoid the n^2 search
+		* through the full edge list
+		*/
+		void FindEdges(void)
+		{
+			// @todo Handle something other than trilists when building edges
+			int32 TriangleCount = Indices.Num() / 3;
+			int32 EdgeCount = 0;
+			// Work through all triangles building the edges
+			for (int32 Triangle = 0; Triangle < TriangleCount; Triangle++)
+			{
+				// Determine the starting index
+				int32 TriangleIndex = Triangle * 3;
+				// Get the indices for the triangle
+				int32 Index1 = Indices[TriangleIndex];
+				int32 Index2 = Indices[TriangleIndex + 1];
+				int32 Index3 = Indices[TriangleIndex + 2];
+				// Add the first to second edge
+				AddEdge(Index1, Index2, Triangle);
+				// Now add the second to third
+				AddEdge(Index2, Index3, Triangle);
+				// Add the third to first edge
+				AddEdge(Index3, Index1, Triangle);
+			}
+		}
+	};
+
+
+	/**
+	* This is the static mesh specific version for finding edges
+	*/
+	class FStaticMeshEdgeBuilder : public TEdgeBuilder<FStaticMeshBuildVertex>
+	{
+	public:
+		/**
+		* Constructor that passes all work to the parent class
+		*/
+		FStaticMeshEdgeBuilder(const TArray<uint32>& InIndices,
+			const TArray<FStaticMeshBuildVertex>& InVertices,
+			TArray<FMeshEdge>& OutEdges) :
+			TEdgeBuilder<FStaticMeshBuildVertex>(InIndices, InVertices, OutEdges)
+		{
+		}
+
+		/**
+		* This function determines whether a given edge matches or not for a static mesh
+		*
+		* @param Index1 The first index of the edge being checked
+		* @param Index2 The second index of the edge
+		* @param OtherEdge The edge to compare. Was found via the map
+		*
+		* @return true if the edge is a match, false otherwise
+		*/
+		inline bool DoesEdgeMatch(int32 Index1, int32 Index2, FMeshEdge* OtherEdge)
+		{
+			return Vertices[OtherEdge->Vertices[1]].Position == Vertices[Index1].Position &&
+				OtherEdge->Faces[1] == -1;
+		}
+	};
+
+	static void BuildDepthOnlyIndexBuffer(
+		TArray<uint32>& OutDepthIndices,
+		const TArray<FStaticMeshBuildVertex>& InVertices,
+		const TArray<uint32>& InIndices,
+		const TArray<FStaticMeshSection>& InSections
+	)
+	{
+		int32 NumVertices = InVertices.Num();
+		if (InIndices.Num() <= 0 || NumVertices <= 0)
+		{
+			OutDepthIndices.Empty();
+			return;
+		}
+
+		// Create a mapping of index -> first overlapping index to accelerate the construction of the shadow index buffer.
+		TArray<FIndexAndZ> VertIndexAndZ;
+		VertIndexAndZ.Empty(NumVertices);
+		for (int32 VertIndex = 0; VertIndex < NumVertices; VertIndex++)
+		{
+			new(VertIndexAndZ)FIndexAndZ(VertIndex, InVertices[VertIndex].Position);
+		}
+		VertIndexAndZ.Sort(FCompareIndexAndZ());
+
+		// Setup the index map. 0xFFFFFFFF == not set.
+		TArray<uint32> IndexMap;
+		IndexMap.AddUninitialized(NumVertices);
+		FMemory::Memset(IndexMap.GetData(), 0xFF, NumVertices * sizeof(uint32));
+
+		// Search for duplicates, quickly!
+		for (int32 i = 0; i < VertIndexAndZ.Num(); i++)
+		{
+			uint32 SrcIndex = VertIndexAndZ[i].Index;
+			float Z = VertIndexAndZ[i].Z;
+			IndexMap[SrcIndex] = FMath::Min(IndexMap[SrcIndex], SrcIndex);
+
+			// Search forward since we add pairs both ways.
+			for (int32 j = i + 1; j < VertIndexAndZ.Num(); j++)
+			{
+				if (FMath::Abs(VertIndexAndZ[j].Z - Z) > THRESH_POINTS_ARE_SAME * 4.01f)
+					break; // can't be any more dups
+
+				uint32 OtherIndex = VertIndexAndZ[j].Index;
+				if (InVertices[SrcIndex].Position.Equals(InVertices[OtherIndex].Position, 0.0f))
+				{
+					IndexMap[SrcIndex] = FMath::Min(IndexMap[SrcIndex], OtherIndex);
+					IndexMap[OtherIndex] = FMath::Min(IndexMap[OtherIndex], SrcIndex);
+				}
+			}
+		}
+
+		// Build the depth-only index buffer by remapping all indices to the first overlapping
+		// vertex in the vertex buffer.
+		OutDepthIndices.Empty();
+		for (int32 SectionIndex = 0; SectionIndex < InSections.Num(); ++SectionIndex)
+		{
+			const FStaticMeshSection& Section = InSections[SectionIndex];
+			int32 FirstIndex = Section.FirstIndex;
+			int32 LastIndex = FirstIndex + Section.NumTriangles * 3;
+			for (int32 SrcIndex = FirstIndex; SrcIndex < LastIndex; ++SrcIndex)
+			{
+				uint32 VertIndex = InIndices[SrcIndex];
+				OutDepthIndices.Add(IndexMap[VertIndex]);
+			}
+		}
+	}
+
+	static void BuildStaticAdjacencyIndexBuffer(
+		const FPositionVertexBuffer& PositionVertexBuffer,
+		const FStaticMeshVertexBuffer& VertexBuffer,
+		const TArray<uint32>& Indices,
+		TArray<uint32>& OutPnAenIndices
+	)
+	{
+		if (Indices.Num())
+		{
+			BuildOptimizationPrivate::NvTriStripHelper::FStaticMeshNvRenderBuffer StaticMeshRenderBuffer(PositionVertexBuffer, VertexBuffer, Indices);
+			nv::IndexBuffer* PnAENIndexBuffer = nv::tess::buildTessellationBuffer(&StaticMeshRenderBuffer, nv::DBM_PnAenDominantCorner, true);
+			check(PnAENIndexBuffer);
+			const int32 IndexCount = (int32)PnAENIndexBuffer->getLength();
+			OutPnAenIndices.Empty(IndexCount);
+			OutPnAenIndices.AddUninitialized(IndexCount);
+			for (int32 Index = 0; Index < IndexCount; ++Index)
+			{
+				OutPnAenIndices[Index] = (*PnAENIndexBuffer)[Index];
+			}
+			delete PnAENIndexBuffer;
+		}
+		else
+		{
+			OutPnAenIndices.Empty();
+		}
+	}
+
+	static void CacheOptimizeIndexBuffer(TArray<uint16>& Indices)
+	{
+		bool bDisableTriangleOrderOptimization = (BuildOptimizationPrivate::CVarTriangleOrderOptimization.GetValueOnGameThread() == 2);
+		bool bUsingNvTriStrip = !bDisableTriangleOrderOptimization && (BuildOptimizationPrivate::CVarTriangleOrderOptimization.GetValueOnGameThread() == 0);
+
+		if (bUsingNvTriStrip)
+		{
+			BuildOptimizationPrivate::NvTriStripHelper::CacheOptimizeIndexBuffer(Indices);
+		}
+		else if (!bDisableTriangleOrderOptimization)
+		{
+			BuildOptimizationPrivate::ForsythHelper::CacheOptimizeIndexBuffer(Indices);
+		}
+	}
+
+	static void CacheOptimizeIndexBuffer(TArray<uint32>& Indices)
+	{
+		bool bDisableTriangleOrderOptimization = (BuildOptimizationPrivate::CVarTriangleOrderOptimization.GetValueOnGameThread() == 2);
+		bool bUsingNvTriStrip = !bDisableTriangleOrderOptimization && (BuildOptimizationPrivate::CVarTriangleOrderOptimization.GetValueOnGameThread() == 0);
+
+		if (bUsingNvTriStrip)
+		{
+			BuildOptimizationPrivate::NvTriStripHelper::CacheOptimizeIndexBuffer(Indices);
+		}
+		else if (!bDisableTriangleOrderOptimization)
+		{
+			BuildOptimizationPrivate::ForsythHelper::CacheOptimizeIndexBuffer(Indices);
+		}
+	}
+
+	static void CacheOptimizeVertexAndIndexBuffer(
+		TArray<FStaticMeshBuildVertex>& Vertices,
+		TArray<TArray<uint32> >& PerSectionIndices
+	)
+	{
+		// Iterate through the section index buffers, 
+		// Optimizing index order for the post transform cache (minimizes the number of vertices transformed), 
+		// And vertex order for the pre transform cache (minimizes the amount of vertex data fetched by the GPU).
+		for (int32 SectionIndex = 0; SectionIndex < PerSectionIndices.Num(); SectionIndex++)
+		{
+			TArray<uint32>& Indices = PerSectionIndices[SectionIndex];
+
+			if (Indices.Num())
+			{
+				// Optimize the index buffer for the post transform cache with.
+				CacheOptimizeIndexBuffer(Indices);
+			}
+		}
+	}
+};
\ No newline at end of file
diff --git a/Engine/Source/Developer/MeshBuilder/Private/BuildStatisticManager.h b/Engine/Source/Developer/MeshBuilder/Private/BuildStatisticManager.h
new file mode 100644
--- /dev/null
+++ b/Engine/Source/Developer/MeshBuilder/Private/BuildStatisticManager.h
@@ -0,0 +1,43 @@
+// Copyright 1998-2017 Epic Games, Inc. All Rights Reserved.
+#pragma once
+
+#include "CoreMinimal.h"
+#include "CoreTypes.h"
+
+DECLARE_LOG_CATEGORY_EXTERN(LogMeshDescriptionBuildStatistic, Log, All);
+
+namespace BuildStatisticManager
+{
+	struct FStatisticData
+	{
+	public:
+		FStatisticData()
+			: TotalTime(0.0)
+			, Counter(0)
+		{ }
+		double TotalTime;
+		uint32 Counter;
+	};
+
+	class FBuildStatisticScope
+	{
+	public:
+		FBuildStatisticScope(const FString& InTimerDescription, FStatisticData& InStatisticData)
+			: TimerDescription(InTimerDescription)
+			, StatisticData(InStatisticData)
+		{
+			StartTime = FPlatformTime::Cycles();
+		}
+		~FBuildStatisticScope()
+		{
+			double ScopeTime = FPlatformTime::ToSeconds64(FPlatformTime::Cycles() - StartTime);
+			StatisticData.TotalTime += ScopeTime;
+			StatisticData.Counter += 1;
+			UE_LOG(LogMeshDescriptionBuildStatistic, Log, TEXT("%s: %f seconds - Counter: %d - Total Time: %f"), *TimerDescription, ScopeTime, StatisticData.Counter, StatisticData.TotalTime);
+		}
+	private:
+		FString TimerDescription;
+		int32 StartTime;
+		FStatisticData& StatisticData;
+	};
+}
\ No newline at end of file
diff --git a/Engine/Source/Developer/MeshBuilder/Private/MeshBuilder.cpp b/Engine/Source/Developer/MeshBuilder/Private/MeshBuilder.cpp
new file mode 100644
--- /dev/null
+++ b/Engine/Source/Developer/MeshBuilder/Private/MeshBuilder.cpp
@@ -0,0 +1,10 @@
+// Copyright 1998-2017 Epic Games, Inc. All Rights Reserved.
+#include "MeshBuilder.h"
+#include "CoreMinimal.h"
+#include "CoreTypes.h"
+
+
+FMeshBuilder::FMeshBuilder()
+{
+
+}
diff --git a/Engine/Source/Developer/MeshBuilder/Private/MeshBuilderModule.cpp b/Engine/Source/Developer/MeshBuilder/Private/MeshBuilderModule.cpp
new file mode 100644
--- /dev/null
+++ b/Engine/Source/Developer/MeshBuilder/Private/MeshBuilderModule.cpp
@@ -0,0 +1,44 @@
+// Copyright 1998-2017 Epic Games, Inc. All Rights Reserved.
+
+#include "IMeshBuilderModule.h"
+#include "Modules/ModuleManager.h"
+
+#include "StaticMeshBuilder.h"
+#include "Engine/StaticMesh.h"
+
+class FMeshBuilderModule : public IMeshBuilderModule
+{
+public:
+
+	FMeshBuilderModule()
+	{
+	}
+
+	virtual void StartupModule() override
+	{
+		// Register any modular features here
+	}
+
+	virtual void ShutdownModule() override
+	{
+		// Unregister any modular features here
+	}
+
+	virtual bool BuildMesh(UObject* Mesh, const FStaticMeshLODGroup& LODGroup) override;
+
+private:
+
+};
+
+IMPLEMENT_MODULE(FMeshBuilderModule, MeshBuilder );
+
+bool FMeshBuilderModule::BuildMesh(class UObject* Mesh, const FStaticMeshLODGroup& LODGroup)
+{
+	UStaticMesh* StaticMesh = Cast<UStaticMesh>(Mesh);
+	if (StaticMesh != nullptr)
+	{
+		//Call the static mesh builder
+		return FStaticMeshBuilder().Build(StaticMesh, LODGroup);
+	}
+	return false;
+}
diff --git a/Engine/Source/Developer/MeshBuilder/Private/MeshDescriptionHelper.cpp b/Engine/Source/Developer/MeshBuilder/Private/MeshDescriptionHelper.cpp
new file mode 100644
--- /dev/null
+++ b/Engine/Source/Developer/MeshBuilder/Private/MeshDescriptionHelper.cpp
@@ -0,0 +1,173 @@
+// Copyright 1998-2017 Epic Games, Inc. All Rights Reserved.
+
+#include "MeshDescriptionHelper.h"
+
+#include "CoreMinimal.h"
+#include "CoreTypes.h"
+#include "UObject/UObjectGlobals.h"
+#include "UObject/Package.h"
+#include "Engine/EngineTypes.h"
+#include "Engine/StaticMesh.h"
+#include "MeshDescription.h"
+#include "MeshAttributes.h"
+#include "Serialization/MemoryWriter.h"
+#include "Serialization/MemoryReader.h"
+#include "RenderUtils.h"
+#include "IMeshReductionInterfaces.h"
+#include "IMeshReductionManagerModule.h"
+#include "Materials/Material.h"
+#include "RawMesh.h"
+#include "BuildStatisticManager.h"
+#include "MeshDescriptionOperations.h"
+
+//Enable all check
+//#define ENABLE_NTB_CHECK
+
+DEFINE_LOG_CATEGORY(LogMeshDescriptionBuildStatistic);
+
+FMeshDescriptionHelper::FMeshDescriptionHelper(FMeshBuildSettings* InBuildSettings, const UMeshDescription* InOriginalMeshDescription)
+	: OriginalMeshDescription(InOriginalMeshDescription)
+	, BuildSettings(InBuildSettings)
+{
+}
+
+UMeshDescription* FMeshDescriptionHelper::GetRenderMeshDescription(UObject* Owner)
+{
+	UStaticMesh* StaticMesh = Cast<UStaticMesh>(Owner);
+	check(StaticMesh);
+	//Use the build settings to create the RenderMeshDescription
+	UMeshDescription *RenderMeshDescription = nullptr;
+
+	if (OriginalMeshDescription == nullptr)
+	{
+		//oups, we do not have a valid original mesh to create the render from
+		return RenderMeshDescription;
+	}
+	//Copy The Original Mesh Description in the render mesh description
+	RenderMeshDescription = Cast<UMeshDescription>(StaticDuplicateObject(OriginalMeshDescription, Owner, NAME_None));
+	float ComparisonThreshold = BuildSettings->bRemoveDegenerates ? THRESH_POINTS_ARE_SAME : 0.0f;
+
+	FVertexInstanceArray& VertexInstanceArray = RenderMeshDescription->VertexInstances();
+	TVertexInstanceAttributeArray<FVector>& Normals = RenderMeshDescription->VertexInstanceAttributes().GetAttributes<FVector>( MeshAttribute::VertexInstance::Normal );
+	TVertexInstanceAttributeArray<FVector>& Tangents = RenderMeshDescription->VertexInstanceAttributes().GetAttributes<FVector>( MeshAttribute::VertexInstance::Tangent );
+	TVertexInstanceAttributeArray<float>& BinormalSigns = RenderMeshDescription->VertexInstanceAttributes().GetAttributes<float>( MeshAttribute::VertexInstance::BinormalSign );
+
+	// Find overlapping corners to accelerate adjacency.
+	FMeshDescriptionOperations::FindOverlappingCorners(OverlappingCorners, RenderMeshDescription, ComparisonThreshold);
+
+	// Compute any missing normals or tangents.
+	{
+		// Static meshes always blend normals of overlapping corners.
+		uint32 TangentOptions = FMeshDescriptionOperations::ETangentOptions::BlendOverlappingNormals;
+		if (BuildSettings->bRemoveDegenerates)
+		{
+			// If removing degenerate triangles, ignore them when computing tangents.
+			TangentOptions |= FMeshDescriptionOperations::ETangentOptions::IgnoreDegenerateTriangles;
+		}
+
+		//This function make sure the Polygon NTB are compute and also remove degenerated triangle from the render mesh description.
+		FMeshDescriptionOperations::CreatePolygonNTB(RenderMeshDescription, BuildSettings->bRemoveDegenerates ? SMALL_NUMBER : 0.0f);
+		
+		//Keep the original mesh description NTBs if we do not rebuild the normals or tangents.
+		bool bComputeTangentLegacy = !BuildSettings->bUseMikkTSpace && (BuildSettings->bRecomputeNormals || BuildSettings->bRecomputeTangents);
+		bool bHasAllNormals = true;
+		bool bHasAllTangents = true;
+		for (const FVertexInstanceID VertexInstanceID : VertexInstanceArray.GetElementIDs())
+		{
+			// Dump normals and tangents if we are recomputing them.
+			if (BuildSettings->bRecomputeTangents)
+			{
+				//Dump the tangents
+				BinormalSigns[VertexInstanceID] = 0.0f;
+				Tangents[VertexInstanceID] = FVector(0.0f);
+			}
+			if (BuildSettings->bRecomputeNormals)
+			{
+				//Dump the normals
+				Normals[VertexInstanceID] = FVector(0.0f);
+			}
+			bHasAllNormals &= !Normals[VertexInstanceID].IsNearlyZero();
+			bHasAllTangents &= !Tangents[VertexInstanceID].IsNearlyZero();
+		}
+
+
+		//MikkTSpace should be use only when the user want to recompute the normals or tangents otherwise should always fallback on builtin
+		//We cannot use mikkt space with degenerated normals fallback on buitin.
+		if (BuildSettings->bUseMikkTSpace && (BuildSettings->bRecomputeNormals || BuildSettings->bRecomputeTangents))
+		{
+			if (!bHasAllNormals)
+			{
+				FMeshDescriptionOperations::CreateNormals(RenderMeshDescription, (FMeshDescriptionOperations::ETangentOptions)TangentOptions, false);
+			}
+			FMeshDescriptionOperations::CreateMikktTangents(RenderMeshDescription, (FMeshDescriptionOperations::ETangentOptions)TangentOptions);
+		}
+		else if(!bHasAllNormals || !bHasAllTangents)
+		{
+			//Set the compute tangent to true when we do not build using mikkt space
+			FMeshDescriptionOperations::CreateNormals(RenderMeshDescription, (FMeshDescriptionOperations::ETangentOptions)TangentOptions, true);
+		}
+	}
+
+	if (BuildSettings->bGenerateLightmapUVs && VertexInstanceArray.Num() > 0)
+	{
+		TVertexInstanceAttributeIndicesArray<FVector2D>& VertexInstanceUVs = RenderMeshDescription->VertexInstanceAttributes().GetAttributesSet<FVector2D>(MeshAttribute::VertexInstance::TextureCoordinate);
+		int32 NumIndices = VertexInstanceUVs.GetNumIndices();
+		//Verify the src light map channel
+		if (BuildSettings->SrcLightmapIndex >= NumIndices)
+		{
+			BuildSettings->SrcLightmapIndex = 0;
+		}
+		//Verify the destination light map channel
+		if (BuildSettings->DstLightmapIndex >= NumIndices)
+		{
+			//Make sure we do not add illegal UV Channel index
+			if (BuildSettings->DstLightmapIndex >= MAX_MESH_TEXTURE_COORDS_MD)
+			{
+				BuildSettings->DstLightmapIndex = MAX_MESH_TEXTURE_COORDS_MD - 1;
+			}
+
+			//Add some unused UVChannel to the mesh description for the lightmapUVs
+			VertexInstanceUVs.SetNumIndices(BuildSettings->DstLightmapIndex + 1);
+			BuildSettings->DstLightmapIndex = NumIndices;
+		}
+		FMeshDescriptionOperations::CreateLightMapUVLayout(RenderMeshDescription,
+			BuildSettings->SrcLightmapIndex,
+			BuildSettings->DstLightmapIndex,
+			BuildSettings->MinLightmapResolution,
+			(FMeshDescriptionOperations::ELightmapUVVersion)(StaticMesh->LightmapUVVersion),
+			OverlappingCorners);
+	}
+
+	return RenderMeshDescription;
+}
+
+void FMeshDescriptionHelper::ReduceLOD(const UMeshDescription* BaseMesh, UMeshDescription* DestMesh, const FMeshReductionSettings& ReductionSettings, const TMultiMap<int32, int32>& InOverlappingCorners)
+{
+	if (BaseMesh == nullptr || DestMesh == nullptr)
+	{
+		//UE_LOG(LogMeshUtilities, Error, TEXT("Mesh contains zero source models."));
+		return;
+	}
+
+	IMeshReductionManagerModule& MeshReductionModule = FModuleManager::Get().LoadModuleChecked<IMeshReductionManagerModule>("MeshReductionInterface");
+	IMeshReduction* MeshReduction = MeshReductionModule.GetStaticMeshReductionInterface();
+	// Reduce this LOD mesh according to its reduction settings.
+
+	if (!MeshReduction || (ReductionSettings.PercentTriangles >= 1.0f && ReductionSettings.MaxDeviation <= 0.0f))
+	{
+		return;
+	}
+	float MaxDeviation = ReductionSettings.MaxDeviation;
+	MeshReduction->ReduceMeshDescription(DestMesh, MaxDeviation, BaseMesh, InOverlappingCorners, ReductionSettings);
+}
+
+bool FMeshDescriptionHelper::IsValidOriginalMeshDescription()
+{
+	return OriginalMeshDescription != nullptr;
+}
+
+void FMeshDescriptionHelper::FindOverlappingCorners(const UMeshDescription* MeshDescription, float ComparisonThreshold)
+{
+	FMeshDescriptionOperations::FindOverlappingCorners(OverlappingCorners, MeshDescription, ComparisonThreshold);
+}
+
diff --git a/Engine/Source/Developer/MeshBuilder/Private/MeshDescriptionHelper.h b/Engine/Source/Developer/MeshBuilder/Private/MeshDescriptionHelper.h
new file mode 100644
--- /dev/null
+++ b/Engine/Source/Developer/MeshBuilder/Private/MeshDescriptionHelper.h
@@ -0,0 +1,67 @@
+// Copyright 1998-2017 Epic Games, Inc. All Rights Reserved.
+#pragma once
+
+#include "CoreMinimal.h"
+#include "CoreTypes.h"
+#include "Logging/LogMacros.h"
+
+DECLARE_LOG_CATEGORY_EXTERN(LogMeshBuilder, Log, All);
+
+class UObject;
+class UMeshDescription;
+struct FMeshBuildSettings;
+struct FVertexInstanceID;
+
+enum
+{
+	//Remove the _MD when FRawMesh will be remove
+	MAX_MESH_TEXTURE_COORDS_MD = 8,
+};
+
+class FMeshDescriptionHelper
+{
+public:
+
+	FMeshDescriptionHelper(FMeshBuildSettings* InBuildSettings, const UMeshDescription* InOriginalMeshDescription);
+
+	//Build a render mesh description with the BuildSettings. This will update the InRenderMeshDescription ptr content
+	UMeshDescription* GetRenderMeshDescription(UObject* Owner);
+
+	void ReduceLOD(const UMeshDescription* BaseMesh, UMeshDescription* DestMesh, const struct FMeshReductionSettings& ReductionSettings, const TMultiMap<int32, int32>& InOverlappingCorners);
+
+	//Return true if there is a valid original mesh description, false otherwise(Auto generate LOD).
+	bool IsValidOriginalMeshDescription();
+
+	void FindOverlappingCorners(const UMeshDescription* MeshDescription, float ComparisonThreshold);
+
+	const TMultiMap<int32, int32>& GetOverlappingCorners() const { return OverlappingCorners; }
+
+private:
+
+	//////////////////////////////////////////////////////////////////////////
+	//PRIVATE function declarations
+
+	//////////////////////////////////////////////////////////////////////////
+	//PRIVATE class members
+
+	const UMeshDescription* OriginalMeshDescription;
+	FMeshBuildSettings *BuildSettings;
+	TMultiMap<int32, int32> OverlappingCorners;
+
+	
+	//////////////////////////////////////////////////////////////////////////
+	//INLINE small helper use to optimize search and compare
+
+	/**
+	* Smoothing group interpretation helper structure.
+	*/
+	struct FFanFace
+	{
+		int32 FaceIndex;
+		int32 LinkedVertexIndex;
+		bool bFilled;
+		bool bBlendTangents;
+		bool bBlendNormals;
+	};
+
+};
diff --git a/Engine/Source/Developer/MeshBuilder/Private/StaticMeshBuilder.cpp b/Engine/Source/Developer/MeshBuilder/Private/StaticMeshBuilder.cpp
new file mode 100644
--- /dev/null
+++ b/Engine/Source/Developer/MeshBuilder/Private/StaticMeshBuilder.cpp
@@ -0,0 +1,633 @@
+// Copyright 1998-2017 Epic Games, Inc. All Rights Reserved.
+
+#include "StaticMeshBuilder.h"
+#include "Engine/StaticMesh.h"
+#include "StaticMeshResources.h"
+#include "PhysicsEngine/BodySetup.h"
+#include "MeshDescription.h"
+#include "MeshDescriptionOperations.h"
+#include "MeshAttributes.h"
+#include "MeshDescriptionHelper.h"
+#include "BuildOptimizationHelper.h"
+#include "Components.h"
+#include "IMeshReductionManagerModule.h"
+#include "MeshBuild.h"
+
+DEFINE_LOG_CATEGORY(LogStaticMeshBuilder);
+
+//////////////////////////////////////////////////////////////////////////
+//Local functions definition
+bool IsOrphanedVertex(const class UMeshDescription *MeshDescription, const FVertexID VertexID);
+void UpdateBounds(class UStaticMesh* StaticMesh);
+void UpdateCollision(UStaticMesh *StaticMesh);
+void BuildVertexBuffer(
+	  UStaticMesh *StaticMesh
+	, int32 LodIndex
+	, UMeshDescription* MeshDescription
+	, FStaticMeshLODResources& StaticMeshLOD
+	, const FMeshBuildSettings& LODBuildSettings
+	, TArray< uint32 >& IndexBuffer
+	, TArray<int32>& OutWedgeMap
+	, TArray<TArray<uint32> >& OutPerSectionIndices
+	, TArray< FStaticMeshBuildVertex >& StaticMeshBuildVertices
+	, const TMultiMap<int32, int32>& OverlappingCorners
+	, float VertexComparisonThreshold
+	, TArray<int32>& RemapVerts);
+void BuildAllBufferOptimizations(struct FStaticMeshLODResources& StaticMeshLOD, const struct FMeshBuildSettings& LODBuildSettings, TArray< uint32 >& IndexBuffer, bool bNeeds32BitIndices, TArray< FStaticMeshBuildVertex >& StaticMeshBuildVertices);
+//////////////////////////////////////////////////////////////////////////
+
+FStaticMeshBuilder::FStaticMeshBuilder()
+{
+
+}
+
+bool FStaticMeshBuilder::Build(UStaticMesh* StaticMesh, const FStaticMeshLODGroup& LODGroup)
+{
+	if (StaticMesh->GetOriginalMeshDescription(0) == nullptr)
+	{
+		//Warn the user that there is no mesh description data
+		UE_LOG(LogStaticMeshBuilder, Error, TEXT("Cannot find a valid mesh description to build the asset."));
+		return false;
+	}
+	StaticMesh->RenderData->AllocateLODResources(StaticMesh->SourceModels.Num());
+
+	//OnBuildRenderMeshStart(StaticMesh, false);
+
+	FStaticMeshRenderData& StaticMeshRenderData = *StaticMesh->RenderData;
+	for (int32 LodIndex = 0; LodIndex < StaticMesh->SourceModels.Num(); ++LodIndex)
+	{
+		FMeshBuildSettings& LODBuildSettings = StaticMesh->SourceModels[LodIndex].BuildSettings;
+		UMeshDescription* OriginalMeshDescription = StaticMesh->GetOriginalMeshDescription(LodIndex);
+		FMeshDescriptionHelper MeshDescriptionHelper(&LODBuildSettings, OriginalMeshDescription);
+		
+		const FStaticMeshSourceModel& SrcModel = StaticMesh->SourceModels[LodIndex];
+		FMeshReductionSettings ReductionSettings = LODGroup.GetSettings(SrcModel.ReductionSettings, LodIndex);
+
+		UMeshDescription* MeshDescription = nullptr;
+		bool bUseReduction = (ReductionSettings.PercentTriangles < 1.0f || ReductionSettings.MaxDeviation > 0.0f) && (StaticMesh->GetMeshDescription(0) != nullptr);
+
+		if (OriginalMeshDescription != nullptr)
+		{
+			OriginalMeshDescription->TriangulateMesh(); //Make sure the original mesh description is triangulate.
+			MeshDescription = MeshDescriptionHelper.GetRenderMeshDescription(StaticMesh);
+		}
+		
+		if(MeshDescription == nullptr)
+		{
+			if (bUseReduction)
+			{
+				//Create an empty mesh description that the reduce will fill
+				MeshDescription = NewObject<UMeshDescription>(StaticMesh, NAME_None);
+			}
+			else if(StaticMesh->GetMeshDescription(0) != nullptr)
+			{
+				//Duplicate the lodindex 0 we have a 100% reduction which is like a duplicate
+				MeshDescription = Cast<UMeshDescription>(StaticDuplicateObject(StaticMesh->GetMeshDescription(0), StaticMesh, NAME_None, RF_NoFlags));
+			}
+		}
+
+		check(MeshDescription != nullptr);
+		StaticMesh->SetMeshDescription(LodIndex, MeshDescription);
+		
+		//Reduce LODs
+		if (bUseReduction)
+		{
+			float OverlappingThreshold = LODBuildSettings.bRemoveDegenerates ? THRESH_POINTS_ARE_SAME : 0.0f;
+			TMultiMap<int32, int32> OverlappingCorners;
+			FMeshDescriptionOperations::FindOverlappingCorners(OverlappingCorners, StaticMesh->GetMeshDescription(0), OverlappingThreshold);
+			MeshDescriptionHelper.ReduceLOD(StaticMesh->GetMeshDescription(0), StaticMesh->GetMeshDescription(LodIndex), ReductionSettings, OverlappingCorners);
+			// Recompute adjacency information. Since we change the vertices when we reduce
+			MeshDescriptionHelper.FindOverlappingCorners(StaticMesh->GetMeshDescription(LodIndex), OverlappingThreshold);
+		}
+
+		const FPolygonGroupArray& PolygonGroups = MeshDescription->PolygonGroups();
+
+		FStaticMeshLODResources& StaticMeshLOD = StaticMeshRenderData.LODResources[LodIndex];
+
+		//discover degenerate triangle with this threshold
+		float VertexComparisonThreshold = LODBuildSettings.bRemoveDegenerates ? THRESH_POINTS_ARE_SAME : 0.0f;
+
+		//Build new vertex buffers
+		TArray< FStaticMeshBuildVertex > StaticMeshBuildVertices;
+
+		TArray< uint32 > IndexBuffer;
+
+		StaticMeshLOD.Sections.Empty(PolygonGroups.Num());
+		TArray<int32> RemapVerts; //Because we will remove MeshVertex that are redundant, we need a remap
+								  //Render data Wedge map is only set for LOD 0???
+		TArray<int32> TempWedgeMap;
+		TArray<int32> &WedgeMap = (LodIndex == 0) ? StaticMeshRenderData.WedgeMap : TempWedgeMap;
+
+		//Prepare the PerSectionIndices array so we can optimize the index buffer for the GPU
+		int32 MaxMaterialIndex = 1;
+		for (const FPolygonGroupID& PolygonGroupID : MeshDescription->PolygonGroups().GetElementIDs())
+		{
+			MaxMaterialIndex = FMath::Max<int32>(PolygonGroupID.GetValue(), MaxMaterialIndex);
+		}
+		TArray<TArray<uint32> > PerSectionIndices;
+		for (int32 i = 0; i <= MaxMaterialIndex; ++i)
+		{
+			PerSectionIndices.Push(TArray<uint32>());
+		}
+
+		//Build the vertex and index buffer
+		BuildVertexBuffer(StaticMesh, LodIndex, MeshDescription, StaticMeshLOD, LODBuildSettings, IndexBuffer, WedgeMap, PerSectionIndices, StaticMeshBuildVertices, MeshDescriptionHelper.GetOverlappingCorners(), VertexComparisonThreshold, RemapVerts);
+
+		// Figure out which index buffer stride we need
+		bool bNeeds32BitIndices = false;
+		for (const FStaticMeshSection& StaticMeshSection : StaticMeshLOD.Sections)
+		{
+			if (StaticMeshSection.MaxVertexIndex > TNumericLimits<uint16>::Max())
+			{
+				bNeeds32BitIndices = true;
+				break;
+			}
+		}
+		const EIndexBufferStride::Type IndexBufferStride = bNeeds32BitIndices ? EIndexBufferStride::Force32Bit : EIndexBufferStride::Force16Bit;
+		StaticMeshLOD.IndexBuffer.SetIndices(IndexBuffer, IndexBufferStride);
+
+		BuildAllBufferOptimizations(StaticMeshLOD, LODBuildSettings, IndexBuffer, bNeeds32BitIndices, StaticMeshBuildVertices);
+	} //End of LOD for loop
+
+	// Calculate the bounding box.
+	FBox BoundingBox(ForceInit);
+	FPositionVertexBuffer& BasePositionVertexBuffer = StaticMeshRenderData.LODResources[0].VertexBuffers.PositionVertexBuffer;
+	for (uint32 VertexIndex = 0; VertexIndex < BasePositionVertexBuffer.GetNumVertices(); VertexIndex++)
+	{
+		BoundingBox += BasePositionVertexBuffer.VertexPosition(VertexIndex);
+	}
+	BoundingBox.GetCenterAndExtents(StaticMeshRenderData.Bounds.Origin, StaticMeshRenderData.Bounds.BoxExtent);
+
+	// Calculate the bounding sphere, using the center of the bounding box as the origin.
+	StaticMeshRenderData.Bounds.SphereRadius = 0.0f;
+	for (uint32 VertexIndex = 0; VertexIndex < BasePositionVertexBuffer.GetNumVertices(); VertexIndex++)
+	{
+		StaticMeshRenderData.Bounds.SphereRadius = FMath::Max(
+			(BasePositionVertexBuffer.VertexPosition(VertexIndex) - StaticMeshRenderData.Bounds.Origin).Size(),
+			StaticMeshRenderData.Bounds.SphereRadius
+		);
+	}
+	//OnBuildRenderMeshFinish(StaticMesh, true);
+
+	return true;
+}
+
+void FStaticMeshBuilder::OnBuildRenderMeshStart(UStaticMesh* StaticMesh, const bool bInvalidateLighting)
+{
+	// We may already have a lock on the rendering resources, if it wasn't released the last time we called EndModification()
+	// on the mesh.  This is only the case when rolling back preview changes for a mesh, because we're guaranteed to apply another
+	// modification to the same mesh in the very same frame.  So we can avoid having to update the GPU resources twice in one frame.
+	if (!RecreateRenderStateContext.IsValid())
+	{
+		// We're changing the mesh itself, so ALL static mesh components in the scene will need
+		// to be unregistered for this (and reregistered afterwards.)
+		const bool bRefreshBounds = true;
+		RecreateRenderStateContext = MakeShareable(new FStaticMeshComponentRecreateRenderStateContext(StaticMesh, bInvalidateLighting, bRefreshBounds));
+
+		// Release the static mesh's resources.
+		StaticMesh->ReleaseResources();
+
+		// Flush the resource release commands to the rendering thread to ensure that the build doesn't occur while a resource is still
+		// allocated, and potentially accessing the UStaticMesh.
+		StaticMesh->ReleaseResourcesFence.Wait();
+	}
+}
+
+void FStaticMeshBuilder::OnBuildRenderMeshFinish(UStaticMesh* StaticMesh, const bool bRebuildBoundsAndCollision)
+{
+	if (bRebuildBoundsAndCollision)
+	{
+		UpdateBounds(StaticMesh);
+		UpdateCollision(StaticMesh);
+	}
+
+	StaticMesh->InitResources();
+
+	// NOTE: This can call InvalidateLightingCache() on all components using this mesh, causing Modify() to be 
+	// called on those components!  Just something to be aware of when EndModification() is called within
+	// an undo transaction.
+	RecreateRenderStateContext.Reset();
+}
+
+bool IsOrphanedVertex(const UMeshDescription *MeshDescription, const FVertexID VertexID)
+{
+	for (const FVertexInstanceID VertexInstanceID : MeshDescription->GetVertexVertexInstances(VertexID))
+	{
+		if (MeshDescription->GetVertexInstanceConnectedPolygons(VertexInstanceID).Num() > 0)
+		{
+			return false;
+		}
+	}
+
+	return true;
+}
+
+void UpdateBounds(UStaticMesh* StaticMesh)
+{
+	const UMeshDescription* MeshDescription = StaticMesh->GetMeshDescription();
+	check(MeshDescription != nullptr);
+
+	// Compute a new bounding box
+	// @todo mesheditor perf: During the final modification, only do this if the bounds may have changed (need hinting)
+	FBoxSphereBounds BoundingBoxAndSphere;
+
+	FBox BoundingBox;
+	BoundingBox.Init();
+
+	// Could improve performance here if necessary:
+	// 1) cache polygon IDs per vertex (in order to quickly reject orphans) and just iterate vertex array; or
+	// 2) cache bounding box per polygon
+	// There are other cases where having polygon adjacency information (1) might be useful, so it's maybe worth considering.
+
+	const FVertexArray& Vertices = MeshDescription->Vertices();
+	const FVertexInstanceArray& VertexInstances = MeshDescription->VertexInstances();
+
+	const TVertexAttributeArray<FVector>& VertexPositions = MeshDescription->VertexAttributes().GetAttributes<FVector>(MeshAttribute::Vertex::Position);
+
+	for (const FVertexID VertexID : Vertices.GetElementIDs())
+	{
+		if (!IsOrphanedVertex(MeshDescription, VertexID))
+		{
+			BoundingBox += VertexPositions[VertexID];
+		}
+	}
+
+	BoundingBox.GetCenterAndExtents( /* Out */ BoundingBoxAndSphere.Origin, /* Out */ BoundingBoxAndSphere.BoxExtent);
+
+	// Calculate the bounding sphere, using the center of the bounding box as the origin.
+	BoundingBoxAndSphere.SphereRadius = 0.0f;
+
+	for (const FVertexID VertexID : Vertices.GetElementIDs())
+	{
+		if (!IsOrphanedVertex(MeshDescription, VertexID))
+		{
+			BoundingBoxAndSphere.SphereRadius = FMath::Max((VertexPositions[VertexID] - BoundingBoxAndSphere.Origin).Size(), BoundingBoxAndSphere.SphereRadius);
+		}
+	}
+
+	FStaticMeshRenderData& StaticMeshRenderData = *StaticMesh->RenderData;
+	StaticMeshRenderData.Bounds = BoundingBoxAndSphere;
+	StaticMesh->CalculateExtendedBounds();
+}
+
+
+void UpdateCollision(UStaticMesh *StaticMesh)
+{
+	// @todo mesheditor collision: We're wiping the existing simplified collision and generating a simple bounding
+	// box collision, since that's the best we can do without impacting performance.  We always using visibility (complex)
+	// collision for traces while mesh editing (for hover/selection), so simplified collision isn't really important.
+	const bool bRecreateSimplifiedCollision = true;
+
+	if (StaticMesh->BodySetup == nullptr)
+	{
+		StaticMesh->CreateBodySetup();
+	}
+
+	UBodySetup* BodySetup = StaticMesh->BodySetup;
+
+	// NOTE: We don't bother calling Modify() on the BodySetup as EndModification() will rebuild this guy after every undo
+	// BodySetup->Modify();
+
+	if (bRecreateSimplifiedCollision)
+	{
+		if (BodySetup->AggGeom.GetElementCount() > 0)
+		{
+			BodySetup->RemoveSimpleCollision();
+		}
+	}
+
+	BodySetup->InvalidatePhysicsData();
+
+	if (bRecreateSimplifiedCollision)
+	{
+		const FBoxSphereBounds Bounds = StaticMesh->GetBounds();
+
+		FKBoxElem BoxElem;
+		BoxElem.Center = Bounds.Origin;
+		BoxElem.X = Bounds.BoxExtent.X * 2.0f;
+		BoxElem.Y = Bounds.BoxExtent.Y * 2.0f;
+		BoxElem.Z = Bounds.BoxExtent.Z * 2.0f;
+		BodySetup->AggGeom.BoxElems.Add(BoxElem);
+	}
+
+	// Update all static mesh components that are using this mesh
+	// @todo mesheditor perf: This is a pretty heavy operation, and overlaps with what we're already doing in RecreateRenderStateContext
+	// a little bit.  Ideally we do everything in a single pass.  Furthermore, if this could be updated lazily it would be faster.
+	{
+		for (FObjectIterator Iter(UStaticMeshComponent::StaticClass()); Iter; ++Iter)
+		{
+			UStaticMeshComponent* StaticMeshComponent = Cast<UStaticMeshComponent>(*Iter);
+			if (StaticMeshComponent->GetStaticMesh() == StaticMesh)
+			{
+				// it needs to recreate IF it already has been created
+				if (StaticMeshComponent->IsPhysicsStateCreated())
+				{
+					StaticMeshComponent->RecreatePhysicsState();
+				}
+			}
+		}
+	}
+}
+
+bool AreVerticesEqual(FStaticMeshBuildVertex const& A, FStaticMeshBuildVertex const& B, float ComparisonThreshold)
+{
+	if (   !A.Position.Equals(B.Position, ComparisonThreshold)
+		|| !NormalsEqual(A.TangentX, B.TangentX)
+		|| !NormalsEqual(A.TangentY, B.TangentY)
+		|| !NormalsEqual(A.TangentZ, B.TangentZ)
+		|| A.Color != B.Color)
+	{
+		return false;
+	}
+
+	// UVs
+	for (int32 UVIndex = 0; UVIndex < MAX_STATIC_TEXCOORDS; UVIndex++)
+	{
+		if (!UVsEqual(A.UVs[UVIndex], B.UVs[UVIndex]))
+		{
+			return false;
+		}
+	}
+
+	return true;
+}
+
+void BuildVertexBuffer(
+	  UStaticMesh *StaticMesh
+	, int32 LodIndex
+	, UMeshDescription* MeshDescription
+	, FStaticMeshLODResources& StaticMeshLOD
+	, const FMeshBuildSettings& LODBuildSettings
+	, TArray< uint32 >& IndexBuffer
+	, TArray<int32>& OutWedgeMap
+	, TArray<TArray<uint32> >& OutPerSectionIndices
+	, TArray< FStaticMeshBuildVertex >& StaticMeshBuildVertices
+	, const TMultiMap<int32, int32>& OverlappingCorners
+	, float VertexComparisonThreshold
+	, TArray<int32>& RemapVerts)
+{
+	FVertexArray& Vertices = MeshDescription->Vertices();
+	FVertexInstanceArray& VertexInstances = MeshDescription->VertexInstances();
+	FPolygonGroupArray& PolygonGroupArray = MeshDescription->PolygonGroups();
+	FPolygonArray& PolygonArray = MeshDescription->Polygons();
+	
+	OutWedgeMap.Reset();
+	OutWedgeMap.AddZeroed(VertexInstances.Num());
+
+	TArray<int32> RemapVertexInstanceID;
+	// set up vertex buffer elements
+	StaticMeshBuildVertices.Reserve(VertexInstances.Num());
+	bool bHasColor = false;
+	//Fill the remap array
+	RemapVerts.AddZeroed(VertexInstances.Num());
+	for (int32& RemapIndex : RemapVerts)
+	{
+		RemapIndex = INDEX_NONE;
+	}
+	TArray<int32> DupVerts;
+
+	const uint32 NumTextureCoord = MeshDescription->VertexInstanceAttributes().GetAttributeIndexCount<FVector2D>(MeshAttribute::VertexInstance::TextureCoordinate);
+
+	const TPolygonGroupAttributeArray<FName>& PolygonGroupImportedMaterialSlotNames = MeshDescription->PolygonGroupAttributes().GetAttributes<FName>(MeshAttribute::PolygonGroup::ImportedMaterialSlotName);
+
+	const TVertexAttributeArray<FVector>& VertexPositions = MeshDescription->VertexAttributes().GetAttributes<FVector>( MeshAttribute::Vertex::Position );
+	const TVertexInstanceAttributeArray<FVector>& VertexInstanceNormals = MeshDescription->VertexInstanceAttributes().GetAttributes<FVector>( MeshAttribute::VertexInstance::Normal );
+	const TVertexInstanceAttributeArray<FVector>& VertexInstanceTangents = MeshDescription->VertexInstanceAttributes().GetAttributes<FVector>( MeshAttribute::VertexInstance::Tangent );
+	const TVertexInstanceAttributeArray<float>& VertexInstanceBinormalSigns = MeshDescription->VertexInstanceAttributes().GetAttributes<float>( MeshAttribute::VertexInstance::BinormalSign );
+	const TVertexInstanceAttributeArray<FVector4>& VertexInstanceColors = MeshDescription->VertexInstanceAttributes().GetAttributes<FVector4>( MeshAttribute::VertexInstance::Color );
+	const TVertexInstanceAttributeIndicesArray<FVector2D>& VertexInstanceUVs = MeshDescription->VertexInstanceAttributes().GetAttributesSet<FVector2D>( MeshAttribute::VertexInstance::TextureCoordinate );
+
+	// Set up index buffer
+	for (const FPolygonGroupID PolygonGroupID : MeshDescription->PolygonGroups().GetElementIDs())
+	{
+		const TArray<FPolygonID>& Polygons = MeshDescription->GetPolygonGroupPolygons(PolygonGroupID);
+
+		TArray<uint32>& SectionIndices = OutPerSectionIndices[PolygonGroupID.GetValue()];
+
+		// Create new rendering section
+		int32 SectionIndex = StaticMeshLOD.Sections.Add(FStaticMeshSection());
+		FStaticMeshSection& StaticMeshSection = StaticMeshLOD.Sections.Last();
+
+		StaticMeshSection.FirstIndex = IndexBuffer.Num();
+		StaticMeshSection.NumTriangles = 0;
+		for (const FPolygonID& PolygonID : Polygons)
+		{
+			StaticMeshSection.NumTriangles += (uint32)MeshDescription->GetPolygonTriangles(PolygonID).Num();
+		}
+
+		StaticMeshSection.MaterialIndex = StaticMesh->GetMaterialIndexFromImportedMaterialSlotName(PolygonGroupImportedMaterialSlotNames[PolygonGroupID]);
+		
+		if (!StaticMesh->StaticMaterials.IsValidIndex(StaticMeshSection.MaterialIndex))
+		{
+			if (StaticMesh->StaticMaterials.IsValidIndex(PolygonGroupID.GetValue()))
+			{
+				StaticMeshSection.MaterialIndex = PolygonGroupID.GetValue();
+			}
+			else
+			{
+				StaticMeshSection.MaterialIndex = 0;
+			}
+		}
+
+		if (LodIndex > 0)
+		{
+			//Set the overwrite section info map
+			FMeshSectionInfo SectionInfo = StaticMesh->SectionInfoMap.Get(LodIndex, SectionIndex);
+			SectionInfo.MaterialIndex = StaticMeshSection.MaterialIndex;
+			StaticMesh->SectionInfoMap.Set(LodIndex, SectionIndex, SectionInfo);
+		}
+
+		for (const FPolygonID PolygonID : Polygons)
+		{
+			const TArray<FMeshTriangle>& PolygonTriangles = MeshDescription->GetPolygonTriangles(PolygonID);
+			int32 ReserveSize = IndexBuffer.Num() + PolygonTriangles.Num() * 3;
+			IndexBuffer.Reserve(ReserveSize);
+			uint32 MinIndex = TNumericLimits< uint32 >::Max();
+			uint32 MaxIndex = TNumericLimits< uint32 >::Min();
+			for (int32 TriangleIndex = 0; TriangleIndex < PolygonTriangles.Num(); ++TriangleIndex)
+			{
+				const FMeshTriangle& Triangle = PolygonTriangles[TriangleIndex];
+				for (int32 TriVert = 0; TriVert < 3; ++TriVert)
+				{
+					const FVertexInstanceID VertexInstanceID = Triangle.GetVertexInstanceID(TriVert);
+					const int32 VertexInstanceValue = VertexInstanceID.GetValue();
+
+					const FVertexID VertexID = MeshDescription->GetVertexInstanceVertex(VertexInstanceID);
+
+					const FVector& VertexPosition = VertexPositions[VertexID];
+					const FVector& VertexInstanceNormal = VertexInstanceNormals[VertexInstanceID];
+					const FVector& VertexInstanceTangent = VertexInstanceTangents[VertexInstanceID];
+					const float VertexInstanceBinormalSign = VertexInstanceBinormalSigns[VertexInstanceID];
+					const FVector4& VertexInstanceColor = VertexInstanceColors[VertexInstanceID];
+
+					const FLinearColor LinearColor(VertexInstanceColor);
+					if (LinearColor != FLinearColor::White)
+					{
+						bHasColor = true;
+					}
+
+					FStaticMeshBuildVertex StaticMeshVertex;
+
+					StaticMeshVertex.Position = VertexPosition * LODBuildSettings.BuildScale3D;
+					const FMatrix ScaleMatrix = FScaleMatrix(LODBuildSettings.BuildScale3D).Inverse().GetTransposed();
+					StaticMeshVertex.TangentX = ScaleMatrix.TransformVector(VertexInstanceTangent).GetSafeNormal();
+					StaticMeshVertex.TangentY = ScaleMatrix.TransformVector(FVector::CrossProduct(VertexInstanceNormal, VertexInstanceTangent).GetSafeNormal() * VertexInstanceBinormalSign).GetSafeNormal();
+					StaticMeshVertex.TangentZ = ScaleMatrix.TransformVector(VertexInstanceNormal).GetSafeNormal();
+					StaticMeshVertex.Color = LinearColor.ToFColor(true);
+					const uint32 MaxNumTexCoords = FMath::Min<int32>(MAX_MESH_TEXTURE_COORDS_MD, MAX_STATIC_TEXCOORDS);
+					for (uint32 UVIndex = 0; UVIndex < MaxNumTexCoords; ++UVIndex)
+					{
+						if(UVIndex < NumTextureCoord)
+						{
+							StaticMeshVertex.UVs[UVIndex] = VertexInstanceUVs.GetArrayForIndex(UVIndex)[VertexInstanceID];
+						}
+						else
+						{
+							StaticMeshVertex.UVs[UVIndex] = FVector2D(0.0f, 0.0f);
+						}
+					}
+					
+
+					//Never add duplicated vertex instance
+					DupVerts.Reset();
+					OverlappingCorners.MultiFind(VertexInstanceValue, DupVerts);
+					int32 Index = INDEX_NONE;
+					for (int32 k = 0; k < DupVerts.Num(); k++)
+					{
+						int32 Location = RemapVerts.IsValidIndex(DupVerts[k]) ? RemapVerts[DupVerts[k]] : INDEX_NONE;
+						if (Location != INDEX_NONE && AreVerticesEqual(StaticMeshVertex, StaticMeshBuildVertices[Location], VertexComparisonThreshold))
+						{
+							Index = Location;
+							break;
+						}
+					}
+					if (Index == INDEX_NONE)
+					{
+						Index = StaticMeshBuildVertices.Add(StaticMeshVertex);
+					}
+					RemapVerts[VertexInstanceValue] = Index;
+					const uint32 RenderingVertexIndex = RemapVerts[VertexInstanceValue];
+					IndexBuffer.Add(RenderingVertexIndex);
+					OutWedgeMap[VertexInstanceValue] = RenderingVertexIndex;
+					SectionIndices.Add(RenderingVertexIndex);
+					MinIndex = FMath::Min(MinIndex, RenderingVertexIndex);
+					MaxIndex = FMath::Max(MaxIndex, RenderingVertexIndex);
+				}
+			}
+			StaticMeshSection.MinVertexIndex = MinIndex;
+			StaticMeshSection.MaxVertexIndex = MaxIndex;
+		}
+
+		if (Polygons.Num() == 0)
+		{
+			// No triangles in this section
+			StaticMeshSection.MinVertexIndex = 0;
+			StaticMeshSection.MaxVertexIndex = 0;
+		}
+	}
+
+	//Optimize before setting the buffer
+	if (VertexInstances.Num() < 100000 * 3)
+	{
+		BuildOptimizationHelper::CacheOptimizeVertexAndIndexBuffer(StaticMeshBuildVertices, OutPerSectionIndices);
+	}
+
+	StaticMeshLOD.VertexBuffers.StaticMeshVertexBuffer.SetUseHighPrecisionTangentBasis(LODBuildSettings.bUseHighPrecisionTangentBasis);
+	StaticMeshLOD.VertexBuffers.StaticMeshVertexBuffer.SetUseFullPrecisionUVs(LODBuildSettings.bUseFullPrecisionUVs);
+	StaticMeshLOD.VertexBuffers.StaticMeshVertexBuffer.Init(StaticMeshBuildVertices, NumTextureCoord);
+	StaticMeshLOD.VertexBuffers.PositionVertexBuffer.Init(StaticMeshBuildVertices);
+	if (bHasColor)
+	{
+		StaticMeshLOD.VertexBuffers.ColorVertexBuffer.Init(StaticMeshBuildVertices);
+	}
+	else
+	{
+		StaticMeshLOD.VertexBuffers.ColorVertexBuffer.InitFromSingleColor(FColor::White, StaticMeshBuildVertices.Num());
+	}
+}
+
+void BuildAllBufferOptimizations(FStaticMeshLODResources& StaticMeshLOD, const FMeshBuildSettings& LODBuildSettings, TArray< uint32 >& IndexBuffer, bool bNeeds32BitIndices, TArray< FStaticMeshBuildVertex >& StaticMeshBuildVertices)
+{
+	const EIndexBufferStride::Type IndexBufferStride = bNeeds32BitIndices ? EIndexBufferStride::Force32Bit : EIndexBufferStride::Force16Bit;
+
+	// Build the reversed index buffer.
+	if (LODBuildSettings.bBuildReversedIndexBuffer)
+	{
+		TArray<uint32> InversedIndices;
+		const int32 IndexCount = IndexBuffer.Num();
+		InversedIndices.AddUninitialized(IndexCount);
+
+		for (int32 SectionIndex = 0; SectionIndex < StaticMeshLOD.Sections.Num(); ++SectionIndex)
+		{
+			const FStaticMeshSection& SectionInfo = StaticMeshLOD.Sections[SectionIndex];
+			const int32 SectionIndexCount = SectionInfo.NumTriangles * 3;
+
+			for (int32 i = 0; i < SectionIndexCount; ++i)
+			{
+				InversedIndices[SectionInfo.FirstIndex + i] = IndexBuffer[SectionInfo.FirstIndex + SectionIndexCount - 1 - i];
+			}
+		}
+		StaticMeshLOD.ReversedIndexBuffer.SetIndices(InversedIndices, IndexBufferStride);
+	}
+
+	// Build the depth-only index buffer.
+	TArray<uint32> DepthOnlyIndices;
+	{
+		BuildOptimizationHelper::BuildDepthOnlyIndexBuffer(
+			DepthOnlyIndices,
+			StaticMeshBuildVertices,
+			IndexBuffer,
+			StaticMeshLOD.Sections
+		);
+
+		if (DepthOnlyIndices.Num() < 50000 * 3)
+		{
+			BuildOptimizationHelper::CacheOptimizeIndexBuffer(DepthOnlyIndices);
+		}
+
+		StaticMeshLOD.DepthOnlyIndexBuffer.SetIndices(DepthOnlyIndices, IndexBufferStride);
+	}
+
+	// Build the inversed depth only index buffer.
+	if (LODBuildSettings.bBuildReversedIndexBuffer)
+	{
+		TArray<uint32> ReversedDepthOnlyIndices;
+		const int32 IndexCount = DepthOnlyIndices.Num();
+		ReversedDepthOnlyIndices.AddUninitialized(IndexCount);
+		for (int32 i = 0; i < IndexCount; ++i)
+		{
+			ReversedDepthOnlyIndices[i] = DepthOnlyIndices[IndexCount - 1 - i];
+		}
+		StaticMeshLOD.ReversedDepthOnlyIndexBuffer.SetIndices(ReversedDepthOnlyIndices, IndexBufferStride);
+	}
+
+	// Build a list of wireframe edges in the static mesh.
+	{
+		TArray<BuildOptimizationHelper::FMeshEdge> Edges;
+		TArray<uint32> WireframeIndices;
+
+		BuildOptimizationHelper::FStaticMeshEdgeBuilder(IndexBuffer, StaticMeshBuildVertices, Edges).FindEdges();
+		WireframeIndices.Empty(2 * Edges.Num());
+		for (int32 EdgeIndex = 0; EdgeIndex < Edges.Num(); EdgeIndex++)
+		{
+			BuildOptimizationHelper::FMeshEdge&	Edge = Edges[EdgeIndex];
+			WireframeIndices.Add(Edge.Vertices[0]);
+			WireframeIndices.Add(Edge.Vertices[1]);
+		}
+		StaticMeshLOD.WireframeIndexBuffer.SetIndices(WireframeIndices, IndexBufferStride);
+	}
+
+	// Build the adjacency index buffer used for tessellation.
+	if (LODBuildSettings.bBuildAdjacencyBuffer)
+	{
+		TArray<uint32> AdjacencyIndices;
+
+		BuildOptimizationHelper::BuildStaticAdjacencyIndexBuffer(
+			StaticMeshLOD.VertexBuffers.PositionVertexBuffer,
+			StaticMeshLOD.VertexBuffers.StaticMeshVertexBuffer,
+			IndexBuffer,
+			AdjacencyIndices
+		);
+		StaticMeshLOD.AdjacencyIndexBuffer.SetIndices(AdjacencyIndices, IndexBufferStride);
+	}
+}
diff --git a/Engine/Source/Developer/MeshBuilder/Private/ThirdPartyBuildOptimizationHelper.h b/Engine/Source/Developer/MeshBuilder/Private/ThirdPartyBuildOptimizationHelper.h
new file mode 100644
--- /dev/null
+++ b/Engine/Source/Developer/MeshBuilder/Private/ThirdPartyBuildOptimizationHelper.h
@@ -0,0 +1,247 @@
+// Copyright 1998-2017 Epic Games, Inc. All Rights Reserved.
+#pragma once
+
+#include "CoreMinimal.h"
+#include "CoreTypes.h"
+//#include "StaticMeshResources.h"
+
+#include "NvTriStrip.h"
+#include "forsythtriangleorderoptimizer.h"
+#include "nvtess.h"
+
+namespace BuildOptimizationPrivate
+{
+	// CVars
+	static TAutoConsoleVariable<int32> CVarTriangleOrderOptimization(
+		TEXT("r.TriangleOrderOptimization_MD"),
+		1,
+		TEXT("Controls the algorithm to use when optimizing the triangle order for the post-transform cache.\n")
+		TEXT("0: Use NVTriStrip (slower)\n")
+		TEXT("1: Use Forsyth algorithm (fastest)(default)")
+		TEXT("2: No triangle order optimization. (least efficient, debugging purposes only)"),
+		ECVF_Default);
+
+	static TAutoConsoleVariable<int32> CVarSupportDepthOnlyIndexBuffers(
+		TEXT("r.SupportDepthOnlyIndexBuffers_MD"),
+		1,
+		TEXT("Enables depth-only index buffers. Saves a little time at the expense of doubling the size of index buffers."),
+		ECVF_ReadOnly | ECVF_RenderThreadSafe);
+
+	static TAutoConsoleVariable<int32> CVarSupportReversedIndexBuffers(
+		TEXT("r.SupportReversedIndexBuffers_MD"),
+		1,
+		TEXT("Enables reversed index buffers. Saves a little time at the expense of doubling the size of index buffers."),
+		ECVF_ReadOnly | ECVF_RenderThreadSafe);
+
+	/*------------------------------------------------------------------------------
+	NVTriStrip for cache optimizing index buffers.
+	------------------------------------------------------------------------------*/
+
+	namespace NvTriStripHelper
+	{
+		/**
+		* Converts 16 bit indices to 32 bit prior to passing them into the real GenerateStrips util method
+		*/
+		void GenerateStrips(
+			const uint8* Indices,
+			bool Is32Bit,
+			const uint32 NumIndices,
+			PrimitiveGroup** PrimGroups,
+			uint32* NumGroups
+		)
+		{
+			if (Is32Bit)
+			{
+				GenerateStrips((uint32*)Indices, NumIndices, PrimGroups, NumGroups);
+			}
+			else
+			{
+				// convert to 32 bit
+				uint32 Idx;
+				TArray<uint32> NewIndices;
+				NewIndices.AddUninitialized(NumIndices);
+				for (Idx = 0; Idx < NumIndices; ++Idx)
+				{
+					NewIndices[Idx] = ((uint16*)Indices)[Idx];
+				}
+				GenerateStrips(NewIndices.GetData(), NumIndices, PrimGroups, NumGroups);
+			}
+
+		}
+
+		/**
+		* Orders a triangle list for better vertex cache coherency.
+		*
+		* *** WARNING: This is safe to call for multiple threads IF AND ONLY IF all
+		* threads call SetListsOnly(true) and SetCacheSize(CACHESIZE_GEFORCE3). If
+		* NvTriStrip is ever used with different settings the library will need
+		* some modifications to be thread-safe. ***
+		*/
+		template<typename IndexDataType, typename Allocator>
+		void CacheOptimizeIndexBuffer(TArray<IndexDataType, Allocator>& Indices)
+		{
+			static_assert(sizeof(IndexDataType) == 2 || sizeof(IndexDataType) == 4, "Indices must be short or int.");
+
+			PrimitiveGroup*	PrimitiveGroups = NULL;
+			uint32			NumPrimitiveGroups = 0;
+			bool Is32Bit = sizeof(IndexDataType) == 4;
+
+			SetListsOnly(true);
+			SetCacheSize(CACHESIZE_GEFORCE3);
+
+			GenerateStrips((uint8*)Indices.GetData(), Is32Bit, Indices.Num(), &PrimitiveGroups, &NumPrimitiveGroups);
+
+			Indices.Empty();
+			Indices.AddUninitialized(PrimitiveGroups->numIndices);
+
+			if (Is32Bit)
+			{
+				FMemory::Memcpy(Indices.GetData(), PrimitiveGroups->indices, Indices.Num() * sizeof(IndexDataType));
+			}
+			else
+			{
+				for (uint32 I = 0; I < PrimitiveGroups->numIndices; ++I)
+				{
+					Indices[I] = (uint16)PrimitiveGroups->indices[I];
+				}
+			}
+
+			delete[] PrimitiveGroups;
+		}
+
+		/**
+		* Provides static mesh render data to the NVIDIA tessellation library.
+		*/
+		class FStaticMeshNvRenderBuffer : public nv::RenderBuffer
+		{
+		public:
+
+			/** Construct from static mesh render buffers. */
+			FStaticMeshNvRenderBuffer(
+				const FPositionVertexBuffer& InPositionVertexBuffer,
+				const FStaticMeshVertexBuffer& InVertexBuffer,
+				const TArray<uint32>& Indices)
+				: PositionVertexBuffer(InPositionVertexBuffer)
+				, VertexBuffer(InVertexBuffer)
+			{
+				check(PositionVertexBuffer.GetNumVertices() == VertexBuffer.GetNumVertices());
+				mIb = new nv::IndexBuffer((void*)Indices.GetData(), nv::IBT_U32, Indices.Num(), false);
+			}
+
+			/** Retrieve the position and first texture coordinate of the specified index. */
+			virtual nv::Vertex getVertex(unsigned int Index) const
+			{
+				nv::Vertex Vertex;
+
+				check(Index < PositionVertexBuffer.GetNumVertices());
+
+				const FVector& Position = PositionVertexBuffer.VertexPosition(Index);
+				Vertex.pos.x = Position.X;
+				Vertex.pos.y = Position.Y;
+				Vertex.pos.z = Position.Z;
+
+				if (VertexBuffer.GetNumTexCoords())
+				{
+					const FVector2D UV = VertexBuffer.GetVertexUV(Index, 0);
+					Vertex.uv.x = UV.X;
+					Vertex.uv.y = UV.Y;
+				}
+				else
+				{
+					Vertex.uv.x = 0.0f;
+					Vertex.uv.y = 0.0f;
+				}
+
+				return Vertex;
+			}
+
+		private:
+
+			/** The position vertex buffer for the static mesh. */
+			const FPositionVertexBuffer& PositionVertexBuffer;
+
+			/** The vertex buffer for the static mesh. */
+			const FStaticMeshVertexBuffer& VertexBuffer;
+
+			/** Copying is forbidden. */
+			FStaticMeshNvRenderBuffer(const FStaticMeshNvRenderBuffer&);
+			FStaticMeshNvRenderBuffer& operator=(const FStaticMeshNvRenderBuffer&);
+		};
+	}
+
+	/*------------------------------------------------------------------------------
+	Forsyth algorithm for cache optimizing index buffers.
+	------------------------------------------------------------------------------*/
+
+	namespace ForsythHelper
+	{
+		/**
+		* Converts 16 bit indices to 32 bit prior to passing them into the real OptimizeFaces util method
+		*/
+		void OptimizeFaces(
+			const uint8* Indices,
+			bool Is32Bit,
+			const uint32 NumIndices,
+			uint32 NumVertices,
+			uint32* OutIndices,
+			uint16 CacheSize
+		)
+		{
+			if (Is32Bit)
+			{
+				Forsyth::OptimizeFaces((uint32*)Indices, NumIndices, NumVertices, OutIndices, CacheSize);
+			}
+			else
+			{
+				// convert to 32 bit
+				uint32 Idx;
+				TArray<uint32> NewIndices;
+				NewIndices.AddUninitialized(NumIndices);
+				for (Idx = 0; Idx < NumIndices; ++Idx)
+				{
+					NewIndices[Idx] = ((uint16*)Indices)[Idx];
+				}
+				Forsyth::OptimizeFaces(NewIndices.GetData(), NumIndices, NumVertices, OutIndices, CacheSize);
+			}
+
+		}
+
+		/**
+		* Orders a triangle list for better vertex cache coherency.
+		*/
+		template<typename IndexDataType, typename Allocator>
+		void CacheOptimizeIndexBuffer(TArray<IndexDataType, Allocator>& Indices)
+		{
+			static_assert(sizeof(IndexDataType) == 2 || sizeof(IndexDataType) == 4, "Indices must be short or int.");
+			bool Is32Bit = sizeof(IndexDataType) == 4;
+
+			// Count the number of vertices
+			uint32 NumVertices = 0;
+			for (int32 Index = 0; Index < Indices.Num(); ++Index)
+			{
+				if (Indices[Index] > NumVertices)
+				{
+					NumVertices = Indices[Index];
+				}
+			}
+			NumVertices += 1;
+
+			TArray<uint32> OptimizedIndices;
+			OptimizedIndices.AddUninitialized(Indices.Num());
+			uint16 CacheSize = 32;
+			OptimizeFaces((uint8*)Indices.GetData(), Is32Bit, Indices.Num(), NumVertices, OptimizedIndices.GetData(), CacheSize);
+
+			if (Is32Bit)
+			{
+				FMemory::Memcpy(Indices.GetData(), OptimizedIndices.GetData(), Indices.Num() * sizeof(IndexDataType));
+			}
+			else
+			{
+				for (int32 I = 0; I < OptimizedIndices.Num(); ++I)
+				{
+					Indices[I] = (uint16)OptimizedIndices[I];
+				}
+			}
+		}
+	}
+}
diff --git a/Engine/Source/Developer/MeshBuilder/Public/IMeshBuilderModule.h b/Engine/Source/Developer/MeshBuilder/Public/IMeshBuilderModule.h
new file mode 100644
--- /dev/null
+++ b/Engine/Source/Developer/MeshBuilder/Public/IMeshBuilderModule.h
@@ -0,0 +1,35 @@
+// Copyright 1998-2017 Epic Games, Inc. All Rights Reserved.
+
+#pragma once
+
+#include "CoreMinimal.h"
+#include "Modules/ModuleInterface.h"
+#include "Modules/ModuleManager.h"
+
+class IMeshBuilderModule : public IModuleInterface
+{
+public:
+
+	/**
+	 * Singleton-like access to this module's interface.  This is just for convenience!
+	 * Beware of calling this during the shutdown phase, though.  Your module might have been unloaded already.
+	 *
+	 * @return Returns singleton instance, loading the module on demand if needed
+	 */
+	static inline IMeshBuilderModule & Get()
+	{
+		return FModuleManager::LoadModuleChecked<IMeshBuilderModule>( "MeshBuilder" );
+	}
+
+	/**
+	 * Checks to see if this module is loaded and ready.  It is only valid to call Get() if IsAvailable() returns true.
+	 *
+	 * @return True if the module is loaded and ready to use
+	 */
+	static inline bool IsAvailable()
+	{
+		return FModuleManager::Get().IsModuleLoaded( "MeshBuilder" );
+	}
+
+	virtual bool BuildMesh(class UObject* Mesh, const class FStaticMeshLODGroup& LODGroup) { return false; }
+};
diff --git a/Engine/Source/Developer/MeshBuilder/Public/MeshBuilder.h b/Engine/Source/Developer/MeshBuilder/Public/MeshBuilder.h
new file mode 100644
--- /dev/null
+++ b/Engine/Source/Developer/MeshBuilder/Public/MeshBuilder.h
@@ -0,0 +1,23 @@
+// Copyright 1998-2017 Epic Games, Inc. All Rights Reserved.
+
+#pragma once
+
+#include "CoreMinimal.h"
+
+/**
+ * Abstract class which is the base class of all builder.
+ * All share code to build some render data should be found inside this class
+ */
+class MESHBUILDER_API FMeshBuilder
+{
+public:
+	FMeshBuilder();
+	
+	/**
+	 * Build function should be override and is the starting point for all mesh builders
+	 */
+	virtual bool Build(class UStaticMesh* StaticMesh, const class FStaticMeshLODGroup& LODGroup) =0;
+
+private:
+
+};
diff --git a/Engine/Source/Developer/MeshBuilder/Public/StaticMeshBuilder.h b/Engine/Source/Developer/MeshBuilder/Public/StaticMeshBuilder.h
new file mode 100644
--- /dev/null
+++ b/Engine/Source/Developer/MeshBuilder/Public/StaticMeshBuilder.h
@@ -0,0 +1,26 @@
+// Copyright 1998-2017 Epic Games, Inc. All Rights Reserved.
+
+#pragma once
+
+#include "CoreMinimal.h"
+#include "MeshBuilder.h"
+
+DECLARE_LOG_CATEGORY_EXTERN(LogStaticMeshBuilder, Log, All);
+
+class MESHBUILDER_API FStaticMeshBuilder : public FMeshBuilder
+{
+public:
+	FStaticMeshBuilder();
+
+	virtual bool Build(class UStaticMesh* StaticMesh, const class FStaticMeshLODGroup& LODGroup) override;
+	virtual ~FStaticMeshBuilder() {}
+
+private:
+
+	void OnBuildRenderMeshStart(class UStaticMesh* StaticMesh, const bool bInvalidateLighting);
+	void OnBuildRenderMeshFinish(class UStaticMesh* StaticMesh, const bool bRebuildBoundsAndCollision);
+
+	/** Used to refresh all components in the scene that may be using a mesh we're editing */
+	TSharedPtr<class FStaticMeshComponentRecreateRenderStateContext> RecreateRenderStateContext;
+};
+
diff --git a/Engine/Source/Runtime/MeshBuilder/MeshBuilder.Build.cs b/Engine/Source/Runtime/MeshBuilder/MeshBuilder.Build.cs
deleted file mode 100644
diff --git a/Engine/Source/Runtime/MeshBuilder/Private/BuildOptimizationHelper.h b/Engine/Source/Runtime/MeshBuilder/Private/BuildOptimizationHelper.h
deleted file mode 100644
diff --git a/Engine/Source/Runtime/MeshBuilder/Private/BuildStatisticManager.h b/Engine/Source/Runtime/MeshBuilder/Private/BuildStatisticManager.h
deleted file mode 100644
diff --git a/Engine/Source/Runtime/MeshBuilder/Private/MeshBuilder.cpp b/Engine/Source/Runtime/MeshBuilder/Private/MeshBuilder.cpp
deleted file mode 100644
diff --git a/Engine/Source/Runtime/MeshBuilder/Private/MeshBuilderModule.cpp b/Engine/Source/Runtime/MeshBuilder/Private/MeshBuilderModule.cpp
deleted file mode 100644
diff --git a/Engine/Source/Runtime/MeshBuilder/Private/MeshDescriptionHelper.cpp b/Engine/Source/Runtime/MeshBuilder/Private/MeshDescriptionHelper.cpp
deleted file mode 100644
diff --git a/Engine/Source/Runtime/MeshBuilder/Private/MeshDescriptionHelper.h b/Engine/Source/Runtime/MeshBuilder/Private/MeshDescriptionHelper.h
deleted file mode 100644
diff --git a/Engine/Source/Runtime/MeshBuilder/Private/StaticMeshBuilder.cpp b/Engine/Source/Runtime/MeshBuilder/Private/StaticMeshBuilder.cpp
deleted file mode 100644
diff --git a/Engine/Source/Runtime/MeshBuilder/Private/ThirdPartyBuildOptimizationHelper.h b/Engine/Source/Runtime/MeshBuilder/Private/ThirdPartyBuildOptimizationHelper.h
deleted file mode 100644
diff --git a/Engine/Source/Runtime/MeshBuilder/Public/IMeshBuilderModule.h b/Engine/Source/Runtime/MeshBuilder/Public/IMeshBuilderModule.h
deleted file mode 100644
diff --git a/Engine/Source/Runtime/MeshBuilder/Public/MeshBuilder.h b/Engine/Source/Runtime/MeshBuilder/Public/MeshBuilder.h
deleted file mode 100644
diff --git a/Engine/Source/Runtime/MeshBuilder/Public/StaticMeshBuilder.h b/Engine/Source/Runtime/MeshBuilder/Public/StaticMeshBuilder.h
deleted file mode 100644
