Fix CIS

diff --git a/Engine/Source/Developer/MeshBuilder/Private/MeshDescriptionHelper.cpp b/Engine/Source/Developer/MeshBuilder/Private/MeshDescriptionHelper.cpp
--- a/Engine/Source/Developer/MeshBuilder/Private/MeshDescriptionHelper.cpp
+++ b/Engine/Source/Developer/MeshBuilder/Private/MeshDescriptionHelper.cpp
@@ -19,6 +19,7 @@
 #include "RawMesh.h"
 #include "BuildStatisticManager.h"
 #include "MeshDescriptionOperations.h"
+#include "Modules/ModuleManager.h"
 
 //Enable all check
 //#define ENABLE_NTB_CHECK
