Fix partialy the automation tests

diff --git a/Engine/Source/Developer/MeshBuilder/Private/StaticMeshBuilder.cpp b/Engine/Source/Developer/MeshBuilder/Private/StaticMeshBuilder.cpp
--- a/Engine/Source/Developer/MeshBuilder/Private/StaticMeshBuilder.cpp
+++ b/Engine/Source/Developer/MeshBuilder/Private/StaticMeshBuilder.cpp
@@ -73,16 +73,28 @@ bool FStaticMeshBuilder::Build(UStaticMesh* StaticMesh, const FStaticMeshLODGrou
 		
 		if(MeshDescription == nullptr)
 		{
+			int32 BaseReduceLodIndex = 0;
 			if (bUseReduction)
 			{
 				//Create an empty mesh description that the reduce will fill
 				MeshDescription = NewObject<UMeshDescription>(StaticMesh, NAME_None);
 				UStaticMesh::RegisterMeshAttributes(MeshDescription);
 			}
-			else if(StaticMesh->GetMeshDescription(0) != nullptr)
+			else if(StaticMesh->GetMeshDescription(BaseReduceLodIndex) != nullptr)
 			{
 				//Duplicate the lodindex 0 we have a 100% reduction which is like a duplicate
-				MeshDescription = Cast<UMeshDescription>(StaticDuplicateObject(StaticMesh->GetMeshDescription(0), StaticMesh, NAME_None, RF_NoFlags));
+				MeshDescription = Cast<UMeshDescription>(StaticDuplicateObject(StaticMesh->GetMeshDescription(BaseReduceLodIndex), StaticMesh, NAME_None, RF_NoFlags));
+			}
+			if (LodIndex > 0)
+			{
+				//Make sure the SectionInfoMap is taken from the Base RawMesh
+				int32 SectionNumber = StaticMesh->OriginalSectionInfoMap.GetSectionNumber(BaseReduceLodIndex);
+				for (int32 SectionIndex = 0; SectionIndex < SectionNumber; ++SectionIndex)
+				{
+					FMeshSectionInfo Info = StaticMesh->OriginalSectionInfoMap.Get(BaseReduceLodIndex, SectionIndex);
+					StaticMesh->SectionInfoMap.Set(LodIndex, SectionIndex, Info);
+					StaticMesh->OriginalSectionInfoMap.Set(LodIndex, SectionIndex, Info);
+				}
 			}
 		}
 
@@ -97,8 +109,53 @@ bool FStaticMeshBuilder::Build(UStaticMesh* StaticMesh, const FStaticMeshLODGrou
 			TMultiMap<int32, int32> OverlappingCorners;
 			FMeshDescriptionOperations::FindOverlappingCorners(OverlappingCorners, StaticMesh->GetMeshDescription(BaseLODIndex), OverlappingThreshold);
 			MeshDescriptionHelper.ReduceLOD(StaticMesh->GetMeshDescription(BaseLODIndex), StaticMesh->GetMeshDescription(LodIndex), ReductionSettings, OverlappingCorners);
-			// Recompute adjacency information. Since we change the vertices when we reduce
-			MeshDescriptionHelper.FindOverlappingCorners(StaticMesh->GetMeshDescription(LodIndex), OverlappingThreshold);
+			UMeshDescription* MeshDescriptionReduced = StaticMesh->GetMeshDescription(LodIndex);
+			if (MeshDescriptionReduced != nullptr)
+			{
+				const TPolygonGroupAttributeArray<FName>& PolygonGroupImportedMaterialSlotNames = MeshDescriptionReduced->PolygonGroupAttributes().GetAttributes<FName>(MeshAttribute::PolygonGroup::ImportedMaterialSlotName);
+				// Recompute adjacency information. Since we change the vertices when we reduce
+				MeshDescriptionHelper.FindOverlappingCorners(MeshDescriptionReduced, OverlappingThreshold);
+			
+				//Make sure the static mesh SectionInfoMap is up to date with the new reduce LOD
+				//We have to remap the material index with the ReductionSettings.BaseLODModel sectionInfoMap
+			
+				//Set the new SectionInfoMap for this reduced LOD base on the ReductionSettings.BaseLODModel SectionInfoMap
+				const FMeshSectionInfoMap& LODModelSectionInfoMap = StaticMesh->SectionInfoMap;
+				TArray<int32> UniqueMaterialIndex;
+				//Find all unique Material in used order
+				for (const FPolygonGroupID& PolygonGroupID : MeshDescriptionReduced->PolygonGroups().GetElementIDs())
+				{
+					int32 MaterialIndex = StaticMesh->GetMaterialIndexFromImportedMaterialSlotName(PolygonGroupImportedMaterialSlotNames[PolygonGroupID]);
+					if (MaterialIndex == INDEX_NONE)
+					{
+						MaterialIndex = PolygonGroupID.GetValue();
+					}
+					UniqueMaterialIndex.AddUnique(MaterialIndex);
+				}
+				//All used material represent a different section
+				for (int32 SectionIndex = 0; SectionIndex < UniqueMaterialIndex.Num(); ++SectionIndex)
+				{
+					//Section material index have to be remap with the ReductionSettings.BaseLODModel SectionInfoMap to create
+					//a valid new section info map for the reduced LOD.
+					if (LODModelSectionInfoMap.IsValidSection(ReductionSettings.BaseLODModel, UniqueMaterialIndex[SectionIndex]))
+					{
+						FMeshSectionInfo SectionInfo = LODModelSectionInfoMap.Get(ReductionSettings.BaseLODModel, UniqueMaterialIndex[SectionIndex]);
+						//Try to recuperate the valid data
+						if (LODModelSectionInfoMap.IsValidSection(LodIndex, SectionIndex))
+						{
+							//If the old LOD section was using the same Material copy the data
+							FMeshSectionInfo OriginalLODSectionInfo = LODModelSectionInfoMap.Get(LodIndex, SectionIndex);
+							if (OriginalLODSectionInfo.MaterialIndex == SectionInfo.MaterialIndex)
+							{
+								SectionInfo.bCastShadow = OriginalLODSectionInfo.bCastShadow;
+								SectionInfo.bEnableCollision = OriginalLODSectionInfo.bEnableCollision;
+							}
+						}
+						//Copy the BaseLODModel section info to the reduce LODIndex.
+						StaticMesh->SectionInfoMap.Set(LodIndex, SectionIndex, SectionInfo);
+					}
+				}
+			}
 		}
 
 		const FPolygonGroupArray& PolygonGroups = MeshDescription->PolygonGroups();
@@ -415,27 +472,11 @@ void BuildVertexBuffer(
 		{
 			StaticMeshSection.NumTriangles += (uint32)MeshDescription->GetPolygonTriangles(PolygonID).Num();
 		}
-
 		StaticMeshSection.MaterialIndex = StaticMesh->GetMaterialIndexFromImportedMaterialSlotName(PolygonGroupImportedMaterialSlotNames[PolygonGroupID]);
 		
-		if (!StaticMesh->StaticMaterials.IsValidIndex(StaticMeshSection.MaterialIndex))
-		{
-			if (StaticMesh->StaticMaterials.IsValidIndex(PolygonGroupID.GetValue()))
-			{
-				StaticMeshSection.MaterialIndex = PolygonGroupID.GetValue();
-			}
-			else
-			{
-				StaticMeshSection.MaterialIndex = 0;
-			}
-		}
-
-		if (LodIndex > 0 && StaticMesh->SectionInfoMap.GetSectionNumber(LodIndex) <= SectionIndex)
+		if (StaticMeshSection.MaterialIndex == INDEX_NONE)
 		{
-			//Set the overwrite section info map in case there is not already one
-			FMeshSectionInfo SectionInfo = StaticMesh->SectionInfoMap.Get(LodIndex, SectionIndex);
-			SectionInfo.MaterialIndex = StaticMeshSection.MaterialIndex;
-			StaticMesh->SectionInfoMap.Set(LodIndex, SectionIndex, SectionInfo);
+			StaticMeshSection.MaterialIndex = PolygonGroupID.GetValue();
 		}
 
 		for (const FPolygonID PolygonID : Polygons)
diff --git a/Engine/Source/Developer/MeshDescriptionOperations/Private/MeshDescriptionOperations.cpp b/Engine/Source/Developer/MeshDescriptionOperations/Private/MeshDescriptionOperations.cpp
--- a/Engine/Source/Developer/MeshDescriptionOperations/Private/MeshDescriptionOperations.cpp
+++ b/Engine/Source/Developer/MeshDescriptionOperations/Private/MeshDescriptionOperations.cpp
@@ -449,30 +449,21 @@ void FMeshDescriptionOperations::ConverFromRawMesh(const FRawMesh &SourceRawMesh
 	TArray<FPolygonGroupID> PolygonGroups;
 	TMap<int32, FPolygonGroupID> MaterialIndexToPolygonGroup;
 
-	//Find the maximum material index
-	int32 MaxMaterialIndex = 0;
+	//Create the PolygonGroups
 	for(int32 MaterialIndex : SourceRawMesh.FaceMaterialIndices)
 	{
-		if (MaterialIndex > MaxMaterialIndex)
+		if (!MaterialIndexToPolygonGroup.Contains(MaterialIndex))
 		{
-			MaxMaterialIndex = MaterialIndex;
-		}
-	}
-
-	//Create the known Materials, this will ensure we create them in the correct order, since it is possible to have FRawMesh with unordered FaceMaterialIndices
-	for (auto Kvp : MaterialMap)
-	{
-		//Avoid creating empty sections
-		if (Kvp.Key > MaxMaterialIndex)
-		{
-			continue;
+			FPolygonGroupID PolygonGroupID(MaterialIndex);
+			DestinationMeshDescription->CreatePolygonGroupWithID(PolygonGroupID);
+			PolygonGroupImportedMaterialSlotNames[PolygonGroupID] = FName(*FString::Printf(TEXT("MaterialSlot_%d"), MaterialIndex));
+			if (MaterialMap.Contains(MaterialIndex))
+			{
+				PolygonGroupImportedMaterialSlotNames[PolygonGroupID] = MaterialMap[MaterialIndex];
+			}
+			PolygonGroups.Add(PolygonGroupID);
+			MaterialIndexToPolygonGroup.Add(MaterialIndex, PolygonGroupID);
 		}
-		FPolygonGroupID PolygonGroupID = DestinationMeshDescription->CreatePolygonGroup();
-		check(Kvp.Key == PolygonGroupID.GetValue());
-		FName PolygonGroupImportedMaterialSlotName = Kvp.Value;
-		PolygonGroupImportedMaterialSlotNames[PolygonGroupID] = PolygonGroupImportedMaterialSlotName == NAME_None ? FName(*FString::Printf(TEXT("MaterialSlot_%d"), Kvp.Key)) : PolygonGroupImportedMaterialSlotName;
-		PolygonGroups.Add(PolygonGroupID);
-		MaterialIndexToPolygonGroup.Add(Kvp.Key, PolygonGroupID);
 	}
 
 	//Triangles
diff --git a/Engine/Source/Editor/UnrealEd/Private/Fbx/FbxStaticMeshImport.cpp b/Engine/Source/Editor/UnrealEd/Private/Fbx/FbxStaticMeshImport.cpp
--- a/Engine/Source/Editor/UnrealEd/Private/Fbx/FbxStaticMeshImport.cpp
+++ b/Engine/Source/Editor/UnrealEd/Private/Fbx/FbxStaticMeshImport.cpp
@@ -1444,6 +1444,32 @@ UStaticMesh* UnFbx::FFbxImporter::ImportStaticMeshAsSingle(UObject* InParent, TA
 					}
 				}
 			}
+			
+			//Set the Section Info Map to fit the real StaticMaterials array
+			int32 SectionIndex = 0;
+			for (const FPolygonGroupID PolygonGroupID : MeshDescription->PolygonGroups().GetElementIDs())
+			{
+				const FName& ImportedMaterialSlotName = PolygonGroupImportedMaterialSlotNames[PolygonGroupID];
+				int32 MaterialIndex = INDEX_NONE;
+				for (int32 FbxMaterialIndex = 0; FbxMaterialIndex < StaticMesh->StaticMaterials.Num(); ++FbxMaterialIndex)
+				{
+					FName& StaticMaterialName = StaticMesh->StaticMaterials[FbxMaterialIndex].ImportedMaterialSlotName;
+					if (StaticMaterialName == ImportedMaterialSlotName)
+					{
+						MaterialIndex = FbxMaterialIndex;
+						break;
+					}
+				}
+				if (MaterialIndex == INDEX_NONE)
+				{
+					MaterialIndex = PolygonGroupID.GetValue();
+				}
+				FMeshSectionInfo Info = StaticMesh->SectionInfoMap.Get(LODIndex, SectionIndex);
+				Info.MaterialIndex = MaterialIndex;
+				StaticMesh->SectionInfoMap.Remove(LODIndex, SectionIndex);
+				StaticMesh->SectionInfoMap.Set(LODIndex, SectionIndex, Info);
+				SectionIndex++;
+			}
 		}
 		//Set the original mesh description to be able to do non destructive reduce
 		StaticMesh->SetOriginalMeshDescription(LODIndex, MeshDescription);
diff --git a/Engine/Source/Runtime/Engine/Private/StaticMesh.cpp b/Engine/Source/Runtime/Engine/Private/StaticMesh.cpp
--- a/Engine/Source/Runtime/Engine/Private/StaticMesh.cpp
+++ b/Engine/Source/Runtime/Engine/Private/StaticMesh.cpp
@@ -1333,7 +1333,7 @@ FArchive& operator<<(FArchive& Ar, FMeshBuildSettings& BuildSettings)
 // differences, etc.) replace the version GUID below with a new one.
 // In case of merge conflicts with DDC versions, you MUST generate a new GUID
 // and set this new GUID as the version.                                       
-#define STATICMESH_DERIVEDDATA_VER TEXT("74CAFC5D929A4BB8AC6B2690A29A1E99")
+#define STATICMESH_DERIVEDDATA_VER TEXT("5299B55CC23A45F58FCEFE9826427C90")
 
 static const FString& GetStaticMeshDerivedDataVersion()
 {
@@ -2689,7 +2689,7 @@ void UStaticMesh::FixupMaterialSlotName()
 // differences, etc.) replace the version GUID below with a new one.
 // In case of merge conflicts with DDC versions, you MUST generate a new GUID
 // and set this new GUID as the version.                                       
-#define MESHDATAKEY_STATICMESH_DERIVEDDATA_VER TEXT("7681842ECB2C4CBDB21613591CDC5B40")
+#define MESHDATAKEY_STATICMESH_DERIVEDDATA_VER TEXT("B452DD73E30F4FC2BC05536BE2E39058")
 
 static const FString& GetMeshDataKeyStaticMeshDerivedDataVersion()
 {
@@ -4227,7 +4227,7 @@ int32 UStaticMesh::GetMaterialIndexFromImportedMaterialSlotName(FName ImportedMa
 			return MaterialIndex;
 		}
 	}
-	return -1;
+	return INDEX_NONE;
 }
 
 /**
