Add log so we can found faulty asset in the cook crash on build machine

diff --git a/Engine/Source/Developer/MeshBuilder/Private/StaticMeshBuilder.cpp b/Engine/Source/Developer/MeshBuilder/Private/StaticMeshBuilder.cpp
--- a/Engine/Source/Developer/MeshBuilder/Private/StaticMeshBuilder.cpp
+++ b/Engine/Source/Developer/MeshBuilder/Private/StaticMeshBuilder.cpp
@@ -49,6 +49,12 @@ bool FStaticMeshBuilder::Build(UStaticMesh* StaticMesh, const FStaticMeshLODGrou
 		UE_LOG(LogStaticMeshBuilder, Error, TEXT("Cannot find a valid mesh description to build the asset."));
 		return false;
 	}
+	if (StaticMesh->RenderData != nullptr && StaticMesh->RenderData->LODResources.Num() > 0)
+	{
+		//This log is to know which staticmesh caugh the error so the log has the full name
+		UE_LOG(LogStaticMeshBuilder, Error, TEXT("StaticMesh RenderData is build 2 time [%s]."), *StaticMesh->GetFullName());
+		return false;
+	}
 	StaticMesh->RenderData->AllocateLODResources(StaticMesh->SourceModels.Num());
 
 	//OnBuildRenderMeshStart(StaticMesh, false);
