Fix the build to use the correct renderdata
Fix a crash in debug when reducing staticmesh

diff --git a/Engine/Source/Developer/MeshBuilder/Private/MeshBuilderModule.cpp b/Engine/Source/Developer/MeshBuilder/Private/MeshBuilderModule.cpp
--- a/Engine/Source/Developer/MeshBuilder/Private/MeshBuilderModule.cpp
+++ b/Engine/Source/Developer/MeshBuilder/Private/MeshBuilderModule.cpp
@@ -24,7 +24,7 @@ public:
 		// Unregister any modular features here
 	}
 
-	virtual bool BuildMesh(UObject* Mesh, const FStaticMeshLODGroup& LODGroup) override;
+	virtual bool BuildMesh(FStaticMeshRenderData& OutRenderData, UObject* Mesh, const FStaticMeshLODGroup& LODGroup) override;
 
 private:
 
@@ -32,13 +32,13 @@ private:
 
 IMPLEMENT_MODULE(FMeshBuilderModule, MeshBuilder );
 
-bool FMeshBuilderModule::BuildMesh(class UObject* Mesh, const FStaticMeshLODGroup& LODGroup)
+bool FMeshBuilderModule::BuildMesh(FStaticMeshRenderData& OutRenderData, class UObject* Mesh, const FStaticMeshLODGroup& LODGroup)
 {
 	UStaticMesh* StaticMesh = Cast<UStaticMesh>(Mesh);
 	if (StaticMesh != nullptr)
 	{
 		//Call the static mesh builder
-		return FStaticMeshBuilder().Build(StaticMesh, LODGroup);
+		return FStaticMeshBuilder().Build(OutRenderData, StaticMesh, LODGroup);
 	}
 	return false;
 }
diff --git a/Engine/Source/Developer/MeshBuilder/Private/StaticMeshBuilder.cpp b/Engine/Source/Developer/MeshBuilder/Private/StaticMeshBuilder.cpp
--- a/Engine/Source/Developer/MeshBuilder/Private/StaticMeshBuilder.cpp
+++ b/Engine/Source/Developer/MeshBuilder/Private/StaticMeshBuilder.cpp
@@ -17,9 +17,6 @@ DEFINE_LOG_CATEGORY(LogStaticMeshBuilder);
 
 //////////////////////////////////////////////////////////////////////////
 //Local functions definition
-bool IsOrphanedVertex(const class UMeshDescription *MeshDescription, const FVertexID VertexID);
-void UpdateBounds(class UStaticMesh* StaticMesh);
-void UpdateCollision(UStaticMesh *StaticMesh);
 void BuildVertexBuffer(
 	  UStaticMesh *StaticMesh
 	, int32 LodIndex
@@ -41,7 +38,7 @@ FStaticMeshBuilder::FStaticMeshBuilder()
 
 }
 
-bool FStaticMeshBuilder::Build(UStaticMesh* StaticMesh, const FStaticMeshLODGroup& LODGroup)
+bool FStaticMeshBuilder::Build(FStaticMeshRenderData& StaticMeshRenderData, UStaticMesh* StaticMesh, const FStaticMeshLODGroup& LODGroup)
 {
 	if (StaticMesh->GetOriginalMeshDescription(0) == nullptr)
 	{
@@ -49,17 +46,19 @@ bool FStaticMeshBuilder::Build(UStaticMesh* StaticMesh, const FStaticMeshLODGrou
 		UE_LOG(LogStaticMeshBuilder, Error, TEXT("Cannot find a valid mesh description to build the asset."));
 		return false;
 	}
-	if (StaticMesh->RenderData != nullptr && StaticMesh->RenderData->LODResources.Num() > 0)
+	if (StaticMeshRenderData.LODResources.Num() > 0)
 	{
-		//This log is to know which staticmesh caugh the error so the log has the full name
-		UE_LOG(LogStaticMeshBuilder, Error, TEXT("StaticMesh RenderData is build 2 time [%s]."), *StaticMesh->GetFullName());
+		//At this point the render data is suppose to be empty
+		UE_LOG(LogStaticMeshBuilder, Error, TEXT("Cannot build static mesh render data twice [%s]."), *StaticMesh->GetFullName());
+		
+		//Crash in debug
+		checkSlow(StaticMeshRenderData.LODResources.Num() == 0);
+
 		return false;
 	}
-	StaticMesh->RenderData->AllocateLODResources(StaticMesh->SourceModels.Num());
+	StaticMeshRenderData.AllocateLODResources(StaticMesh->SourceModels.Num());
 
-	//OnBuildRenderMeshStart(StaticMesh, false);
 
-	FStaticMeshRenderData& StaticMeshRenderData = *StaticMesh->RenderData;
 	for (int32 LodIndex = 0; LodIndex < StaticMesh->SourceModels.Num(); ++LodIndex)
 	{
 		FMeshBuildSettings& LODBuildSettings = StaticMesh->SourceModels[LodIndex].BuildSettings;
@@ -231,168 +230,10 @@ bool FStaticMeshBuilder::Build(UStaticMesh* StaticMesh, const FStaticMeshLODGrou
 			StaticMeshRenderData.Bounds.SphereRadius
 		);
 	}
-	//OnBuildRenderMeshFinish(StaticMesh, true);
-
-	return true;
-}
-
-void FStaticMeshBuilder::OnBuildRenderMeshStart(UStaticMesh* StaticMesh, const bool bInvalidateLighting)
-{
-	// We may already have a lock on the rendering resources, if it wasn't released the last time we called EndModification()
-	// on the mesh.  This is only the case when rolling back preview changes for a mesh, because we're guaranteed to apply another
-	// modification to the same mesh in the very same frame.  So we can avoid having to update the GPU resources twice in one frame.
-	if (!RecreateRenderStateContext.IsValid())
-	{
-		// We're changing the mesh itself, so ALL static mesh components in the scene will need
-		// to be unregistered for this (and reregistered afterwards.)
-		const bool bRefreshBounds = true;
-		RecreateRenderStateContext = MakeShareable(new FStaticMeshComponentRecreateRenderStateContext(StaticMesh, bInvalidateLighting, bRefreshBounds));
-
-		// Release the static mesh's resources.
-		StaticMesh->ReleaseResources();
-
-		// Flush the resource release commands to the rendering thread to ensure that the build doesn't occur while a resource is still
-		// allocated, and potentially accessing the UStaticMesh.
-		StaticMesh->ReleaseResourcesFence.Wait();
-	}
-}
-
-void FStaticMeshBuilder::OnBuildRenderMeshFinish(UStaticMesh* StaticMesh, const bool bRebuildBoundsAndCollision)
-{
-	if (bRebuildBoundsAndCollision)
-	{
-		UpdateBounds(StaticMesh);
-		UpdateCollision(StaticMesh);
-	}
-
-	StaticMesh->InitResources();
-
-	// NOTE: This can call InvalidateLightingCache() on all components using this mesh, causing Modify() to be 
-	// called on those components!  Just something to be aware of when EndModification() is called within
-	// an undo transaction.
-	RecreateRenderStateContext.Reset();
-}
-
-bool IsOrphanedVertex(const UMeshDescription *MeshDescription, const FVertexID VertexID)
-{
-	for (const FVertexInstanceID VertexInstanceID : MeshDescription->GetVertexVertexInstances(VertexID))
-	{
-		if (MeshDescription->GetVertexInstanceConnectedPolygons(VertexInstanceID).Num() > 0)
-		{
-			return false;
-		}
-	}
 
 	return true;
 }
 
-void UpdateBounds(UStaticMesh* StaticMesh)
-{
-	const UMeshDescription* MeshDescription = StaticMesh->GetMeshDescription();
-	check(MeshDescription != nullptr);
-
-	// Compute a new bounding box
-	// @todo mesheditor perf: During the final modification, only do this if the bounds may have changed (need hinting)
-	FBoxSphereBounds BoundingBoxAndSphere;
-
-	FBox BoundingBox;
-	BoundingBox.Init();
-
-	// Could improve performance here if necessary:
-	// 1) cache polygon IDs per vertex (in order to quickly reject orphans) and just iterate vertex array; or
-	// 2) cache bounding box per polygon
-	// There are other cases where having polygon adjacency information (1) might be useful, so it's maybe worth considering.
-
-	const FVertexArray& Vertices = MeshDescription->Vertices();
-	const FVertexInstanceArray& VertexInstances = MeshDescription->VertexInstances();
-
-	const TVertexAttributeArray<FVector>& VertexPositions = MeshDescription->VertexAttributes().GetAttributes<FVector>(MeshAttribute::Vertex::Position);
-
-	for (const FVertexID VertexID : Vertices.GetElementIDs())
-	{
-		if (!IsOrphanedVertex(MeshDescription, VertexID))
-		{
-			BoundingBox += VertexPositions[VertexID];
-		}
-	}
-
-	BoundingBox.GetCenterAndExtents( /* Out */ BoundingBoxAndSphere.Origin, /* Out */ BoundingBoxAndSphere.BoxExtent);
-
-	// Calculate the bounding sphere, using the center of the bounding box as the origin.
-	BoundingBoxAndSphere.SphereRadius = 0.0f;
-
-	for (const FVertexID VertexID : Vertices.GetElementIDs())
-	{
-		if (!IsOrphanedVertex(MeshDescription, VertexID))
-		{
-			BoundingBoxAndSphere.SphereRadius = FMath::Max((VertexPositions[VertexID] - BoundingBoxAndSphere.Origin).Size(), BoundingBoxAndSphere.SphereRadius);
-		}
-	}
-
-	FStaticMeshRenderData& StaticMeshRenderData = *StaticMesh->RenderData;
-	StaticMeshRenderData.Bounds = BoundingBoxAndSphere;
-	StaticMesh->CalculateExtendedBounds();
-}
-
-
-void UpdateCollision(UStaticMesh *StaticMesh)
-{
-	// @todo mesheditor collision: We're wiping the existing simplified collision and generating a simple bounding
-	// box collision, since that's the best we can do without impacting performance.  We always using visibility (complex)
-	// collision for traces while mesh editing (for hover/selection), so simplified collision isn't really important.
-	const bool bRecreateSimplifiedCollision = true;
-
-	if (StaticMesh->BodySetup == nullptr)
-	{
-		StaticMesh->CreateBodySetup();
-	}
-
-	UBodySetup* BodySetup = StaticMesh->BodySetup;
-
-	// NOTE: We don't bother calling Modify() on the BodySetup as EndModification() will rebuild this guy after every undo
-	// BodySetup->Modify();
-
-	if (bRecreateSimplifiedCollision)
-	{
-		if (BodySetup->AggGeom.GetElementCount() > 0)
-		{
-			BodySetup->RemoveSimpleCollision();
-		}
-	}
-
-	BodySetup->InvalidatePhysicsData();
-
-	if (bRecreateSimplifiedCollision)
-	{
-		const FBoxSphereBounds Bounds = StaticMesh->GetBounds();
-
-		FKBoxElem BoxElem;
-		BoxElem.Center = Bounds.Origin;
-		BoxElem.X = Bounds.BoxExtent.X * 2.0f;
-		BoxElem.Y = Bounds.BoxExtent.Y * 2.0f;
-		BoxElem.Z = Bounds.BoxExtent.Z * 2.0f;
-		BodySetup->AggGeom.BoxElems.Add(BoxElem);
-	}
-
-	// Update all static mesh components that are using this mesh
-	// @todo mesheditor perf: This is a pretty heavy operation, and overlaps with what we're already doing in RecreateRenderStateContext
-	// a little bit.  Ideally we do everything in a single pass.  Furthermore, if this could be updated lazily it would be faster.
-	{
-		for (FObjectIterator Iter(UStaticMeshComponent::StaticClass()); Iter; ++Iter)
-		{
-			UStaticMeshComponent* StaticMeshComponent = Cast<UStaticMeshComponent>(*Iter);
-			if (StaticMeshComponent->GetStaticMesh() == StaticMesh)
-			{
-				// it needs to recreate IF it already has been created
-				if (StaticMeshComponent->IsPhysicsStateCreated())
-				{
-					StaticMeshComponent->RecreatePhysicsState();
-				}
-			}
-		}
-	}
-}
-
 bool AreVerticesEqual(FStaticMeshBuildVertex const& A, FStaticMeshBuildVertex const& B, float ComparisonThreshold)
 {
 	if (   !A.Position.Equals(B.Position, ComparisonThreshold)
diff --git a/Engine/Source/Developer/MeshBuilder/Public/IMeshBuilderModule.h b/Engine/Source/Developer/MeshBuilder/Public/IMeshBuilderModule.h
--- a/Engine/Source/Developer/MeshBuilder/Public/IMeshBuilderModule.h
+++ b/Engine/Source/Developer/MeshBuilder/Public/IMeshBuilderModule.h
@@ -31,5 +31,5 @@ public:
 		return FModuleManager::Get().IsModuleLoaded( "MeshBuilder" );
 	}
 
-	virtual bool BuildMesh(class UObject* Mesh, const class FStaticMeshLODGroup& LODGroup) { return false; }
+	virtual bool BuildMesh(class FStaticMeshRenderData& OutRenderData, class UObject* Mesh, const class FStaticMeshLODGroup& LODGroup) { return false; }
 };
diff --git a/Engine/Source/Developer/MeshBuilder/Public/MeshBuilder.h b/Engine/Source/Developer/MeshBuilder/Public/MeshBuilder.h
--- a/Engine/Source/Developer/MeshBuilder/Public/MeshBuilder.h
+++ b/Engine/Source/Developer/MeshBuilder/Public/MeshBuilder.h
@@ -16,7 +16,7 @@ public:
 	/**
 	 * Build function should be override and is the starting point for all mesh builders
 	 */
-	virtual bool Build(class UStaticMesh* StaticMesh, const class FStaticMeshLODGroup& LODGroup) =0;
+	virtual bool Build(class FStaticMeshRenderData& OutRenderData, class UStaticMesh* StaticMesh, const class FStaticMeshLODGroup& LODGroup) =0;
 
 private:
 
diff --git a/Engine/Source/Developer/MeshBuilder/Public/StaticMeshBuilder.h b/Engine/Source/Developer/MeshBuilder/Public/StaticMeshBuilder.h
--- a/Engine/Source/Developer/MeshBuilder/Public/StaticMeshBuilder.h
+++ b/Engine/Source/Developer/MeshBuilder/Public/StaticMeshBuilder.h
@@ -12,7 +12,7 @@ class MESHBUILDER_API FStaticMeshBuilder : public FMeshBuilder
 public:
 	FStaticMeshBuilder();
 
-	virtual bool Build(class UStaticMesh* StaticMesh, const class FStaticMeshLODGroup& LODGroup) override;
+	virtual bool Build(class FStaticMeshRenderData& OutRenderData, class UStaticMesh* StaticMesh, const class FStaticMeshLODGroup& LODGroup) override;
 	virtual ~FStaticMeshBuilder() {}
 
 private:
diff --git a/Engine/Source/Developer/MeshSimplifier/Private/Quadric.h b/Engine/Source/Developer/MeshSimplifier/Private/Quadric.h
--- a/Engine/Source/Developer/MeshSimplifier/Private/Quadric.h
+++ b/Engine/Source/Developer/MeshSimplifier/Private/Quadric.h
@@ -813,8 +813,10 @@ inline bool TQuadricAttrOptimizer< NumAttributes >::Optimize( FVector& Point ) c
 	//     [ -d[ 0 .. m] ]
 
 	// ( C - 1/a * B*Bt ) * p = -1/a * B*d - dn
-
-	checkSlow( a != 0.0 );
+	if (FMath::IsNearlyZero(a, (double)(1.e-12)))
+	{
+		return false;
+	}
 	
 	// M = C - 1/a * B*Bt
 	double ia = 1.0 / a;
diff --git a/Engine/Source/Runtime/Engine/Private/StaticMesh.cpp b/Engine/Source/Runtime/Engine/Private/StaticMesh.cpp
--- a/Engine/Source/Runtime/Engine/Private/StaticMesh.cpp
+++ b/Engine/Source/Runtime/Engine/Private/StaticMesh.cpp
@@ -1591,7 +1591,7 @@ void FStaticMeshRenderData::Cache(UStaticMesh* Owner, const FStaticMeshLODSettin
 			if (bUseMeshDescription)
 			{
 				IMeshBuilderModule& MeshBuilderModule = FModuleManager::Get().LoadModuleChecked<IMeshBuilderModule>(TEXT("MeshBuilder"));
-				if (!MeshBuilderModule.BuildMesh(Owner, LODGroup))
+				if (!MeshBuilderModule.BuildMesh(*this, Owner, LODGroup))
 				{
 					UE_LOG(LogStaticMesh, Error, TEXT("Failed to build static mesh. See previous line(s) for details."));
 					return;
