Fix rendering translucency automation tests

diff --git a/Engine/Source/Developer/MeshBuilder/Private/BuildOptimizationHelper.h b/Engine/Source/Developer/MeshBuilder/Private/BuildOptimizationHelper.h
--- a/Engine/Source/Developer/MeshBuilder/Private/BuildOptimizationHelper.h
+++ b/Engine/Source/Developer/MeshBuilder/Private/BuildOptimizationHelper.h
@@ -344,9 +344,19 @@ public:
 
 	static void CacheOptimizeVertexAndIndexBuffer(
 		TArray<FStaticMeshBuildVertex>& Vertices,
-		TArray<TArray<uint32> >& PerSectionIndices
+		TArray<TArray<uint32> >& PerSectionIndices,
+		TArray<int32>& WedgeMap
 	)
 	{
+		// Copy the vertices since we will be reordering them
+		TArray<FStaticMeshBuildVertex> OriginalVertices = Vertices;
+
+		// Initialize a cache that stores which indices have been assigned
+		TArray<int32> IndexCache;
+		IndexCache.AddUninitialized(Vertices.Num());
+		FMemory::Memset(IndexCache.GetData(), INDEX_NONE, IndexCache.Num() * IndexCache.GetTypeSize());
+		int32 NextAvailableIndex = 0;
+
 		// Iterate through the section index buffers, 
 		// Optimizing index order for the post transform cache (minimizes the number of vertices transformed), 
 		// And vertex order for the pre transform cache (minimizes the amount of vertex data fetched by the GPU).
@@ -358,6 +368,40 @@ public:
 			{
 				// Optimize the index buffer for the post transform cache with.
 				CacheOptimizeIndexBuffer(Indices);
+
+				// Copy the index buffer since we will be reordering it
+				TArray<uint32> OriginalIndices = Indices;
+
+				// Go through the indices and assign them new values that are coherent where possible
+				for (int32 Index = 0; Index < Indices.Num(); Index++)
+				{
+					const int32 CachedIndex = IndexCache[OriginalIndices[Index]];
+
+					if (CachedIndex == INDEX_NONE)
+					{
+						// No new index has been allocated for this existing index, assign a new one
+						Indices[Index] = NextAvailableIndex;
+						// Mark what this index has been assigned to
+						IndexCache[OriginalIndices[Index]] = NextAvailableIndex;
+						NextAvailableIndex++;
+					}
+					else
+					{
+						// Reuse an existing index assignment
+						Indices[Index] = CachedIndex;
+					}
+					// Reorder the vertices based on the new index assignment
+					Vertices[Indices[Index]] = OriginalVertices[OriginalIndices[Index]];
+				}
+			}
+		}
+
+		for (int32 i = 0; i < WedgeMap.Num(); i++)
+		{
+			int32 MappedIndex = WedgeMap[i];
+			if (MappedIndex != INDEX_NONE)
+			{
+				WedgeMap[i] = IndexCache[MappedIndex];
 			}
 		}
 	}
diff --git a/Engine/Source/Developer/MeshBuilder/Private/StaticMeshBuilder.cpp b/Engine/Source/Developer/MeshBuilder/Private/StaticMeshBuilder.cpp
--- a/Engine/Source/Developer/MeshBuilder/Private/StaticMeshBuilder.cpp
+++ b/Engine/Source/Developer/MeshBuilder/Private/StaticMeshBuilder.cpp
@@ -196,20 +196,46 @@ bool FStaticMeshBuilder::Build(FStaticMeshRenderData& StaticMeshRenderData, USta
 		//Build the vertex and index buffer
 		BuildVertexBuffer(StaticMesh, LodIndex, MeshDescription, StaticMeshLOD, LODBuildSettings, IndexBuffer, WedgeMap, PerSectionIndices, StaticMeshBuildVertices, MeshDescriptionHelper.GetOverlappingCorners(), VertexComparisonThreshold, RemapVerts);
 
-		// Figure out which index buffer stride we need
+		// Concatenate the per-section index buffers.
+		TArray<uint32> CombinedIndices;
 		bool bNeeds32BitIndices = false;
-		for (const FStaticMeshSection& StaticMeshSection : StaticMeshLOD.Sections)
+		for (int32 SectionIndex = 0; SectionIndex < StaticMeshLOD.Sections.Num(); SectionIndex++)
 		{
-			if (StaticMeshSection.MaxVertexIndex > TNumericLimits<uint16>::Max())
+			FStaticMeshSection& Section = StaticMeshLOD.Sections[SectionIndex];
+			TArray<uint32> const& SectionIndices = PerSectionIndices[SectionIndex];
+			Section.FirstIndex = 0;
+			Section.NumTriangles = 0;
+			Section.MinVertexIndex = 0;
+			Section.MaxVertexIndex = 0;
+
+			if (SectionIndices.Num())
 			{
-				bNeeds32BitIndices = true;
-				break;
+				Section.FirstIndex = CombinedIndices.Num();
+				Section.NumTriangles = SectionIndices.Num() / 3;
+
+				CombinedIndices.AddUninitialized(SectionIndices.Num());
+				uint32* DestPtr = &CombinedIndices[Section.FirstIndex];
+				uint32 const* SrcPtr = SectionIndices.GetData();
+
+				Section.MinVertexIndex = *SrcPtr;
+				Section.MaxVertexIndex = *SrcPtr;
+
+				for (int32 Index = 0; Index < SectionIndices.Num(); Index++)
+				{
+					uint32 VertIndex = *SrcPtr++;
+
+					bNeeds32BitIndices |= (VertIndex > MAX_uint16);
+					Section.MinVertexIndex = FMath::Min<uint32>(VertIndex, Section.MinVertexIndex);
+					Section.MaxVertexIndex = FMath::Max<uint32>(VertIndex, Section.MaxVertexIndex);
+					*DestPtr++ = VertIndex;
+				}
 			}
 		}
+
 		const EIndexBufferStride::Type IndexBufferStride = bNeeds32BitIndices ? EIndexBufferStride::Force32Bit : EIndexBufferStride::Force16Bit;
-		StaticMeshLOD.IndexBuffer.SetIndices(IndexBuffer, IndexBufferStride);
+		StaticMeshLOD.IndexBuffer.SetIndices(CombinedIndices, IndexBufferStride);
 
-		BuildAllBufferOptimizations(StaticMeshLOD, LODBuildSettings, IndexBuffer, bNeeds32BitIndices, StaticMeshBuildVertices);
+		BuildAllBufferOptimizations(StaticMeshLOD, LODBuildSettings, CombinedIndices, bNeeds32BitIndices, StaticMeshBuildVertices);
 	} //End of LOD for loop
 
 	// Calculate the bounding box.
@@ -418,7 +444,8 @@ void BuildVertexBuffer(
 	//Optimize before setting the buffer
 	if (VertexInstances.Num() < 100000 * 3)
 	{
-		BuildOptimizationHelper::CacheOptimizeVertexAndIndexBuffer(StaticMeshBuildVertices, OutPerSectionIndices);
+		BuildOptimizationHelper::CacheOptimizeVertexAndIndexBuffer(StaticMeshBuildVertices, OutPerSectionIndices, OutWedgeMap);
+		check(OutWedgeMap.Num() == MeshDescription->VertexInstances().Num());
 	}
 
 	StaticMeshLOD.VertexBuffers.StaticMeshVertexBuffer.SetUseHighPrecisionTangentBasis(LODBuildSettings.bUseHighPrecisionTangentBasis);
diff --git a/Engine/Source/Runtime/Engine/Private/StaticMesh.cpp b/Engine/Source/Runtime/Engine/Private/StaticMesh.cpp
--- a/Engine/Source/Runtime/Engine/Private/StaticMesh.cpp
+++ b/Engine/Source/Runtime/Engine/Private/StaticMesh.cpp
@@ -1333,7 +1333,7 @@ FArchive& operator<<(FArchive& Ar, FMeshBuildSettings& BuildSettings)
 // differences, etc.) replace the version GUID below with a new one.
 // In case of merge conflicts with DDC versions, you MUST generate a new GUID
 // and set this new GUID as the version.                                       
-#define STATICMESH_DERIVEDDATA_VER TEXT("8A8FE832D30D4F72A111D4FF9BCF4346")
+#define STATICMESH_DERIVEDDATA_VER TEXT("B330D0DD8F214096995A2A9527F4E0DD")
 
 static const FString& GetStaticMeshDerivedDataVersion()
 {
@@ -2689,7 +2689,7 @@ void UStaticMesh::FixupMaterialSlotName()
 // differences, etc.) replace the version GUID below with a new one.
 // In case of merge conflicts with DDC versions, you MUST generate a new GUID
 // and set this new GUID as the version.                                       
-#define MESHDATAKEY_STATICMESH_DERIVEDDATA_VER TEXT("24E2D709773D4503B6873C4BED85ECB2")
+#define MESHDATAKEY_STATICMESH_DERIVEDDATA_VER TEXT("22EA1C022AC745028AB8AB4904B15B75")
 
 static const FString& GetMeshDataKeyStaticMeshDerivedDataVersion()
 {
