Fix the mesh merge tool to fix automation tests

diff --git a/Engine/Source/Developer/MeshBuilder/Private/StaticMeshBuilder.cpp b/Engine/Source/Developer/MeshBuilder/Private/StaticMeshBuilder.cpp
--- a/Engine/Source/Developer/MeshBuilder/Private/StaticMeshBuilder.cpp
+++ b/Engine/Source/Developer/MeshBuilder/Private/StaticMeshBuilder.cpp
@@ -182,16 +182,8 @@ bool FStaticMeshBuilder::Build(FStaticMeshRenderData& StaticMeshRenderData, USta
 		TArray<int32> &WedgeMap = (LodIndex == 0) ? StaticMeshRenderData.WedgeMap : TempWedgeMap;
 
 		//Prepare the PerSectionIndices array so we can optimize the index buffer for the GPU
-		int32 MaxMaterialIndex = 1;
-		for (const FPolygonGroupID& PolygonGroupID : MeshDescription->PolygonGroups().GetElementIDs())
-		{
-			MaxMaterialIndex = FMath::Max<int32>(PolygonGroupID.GetValue(), MaxMaterialIndex);
-		}
 		TArray<TArray<uint32> > PerSectionIndices;
-		for (int32 i = 0; i <= MaxMaterialIndex; ++i)
-		{
-			PerSectionIndices.Push(TArray<uint32>());
-		}
+		PerSectionIndices.AddDefaulted(MeshDescription->PolygonGroups().Num());
 
 		//Build the vertex and index buffer
 		BuildVertexBuffer(StaticMesh, LodIndex, MeshDescription, StaticMeshLOD, LODBuildSettings, IndexBuffer, WedgeMap, PerSectionIndices, StaticMeshBuildVertices, MeshDescriptionHelper.GetOverlappingCorners(), VertexComparisonThreshold, RemapVerts);
@@ -333,11 +325,10 @@ void BuildVertexBuffer(
 	{
 		const TArray<FPolygonID>& Polygons = MeshDescription->GetPolygonGroupPolygons(PolygonGroupID);
 
-		TArray<uint32>& SectionIndices = OutPerSectionIndices[PolygonGroupID.GetValue()];
-
 		// Create new rendering section
 		int32 SectionIndex = StaticMeshLOD.Sections.Add(FStaticMeshSection());
 		FStaticMeshSection& StaticMeshSection = StaticMeshLOD.Sections.Last();
+		TArray<uint32>& SectionIndices = OutPerSectionIndices[SectionIndex];
 
 		StaticMeshSection.FirstIndex = IndexBuffer.Num();
 		StaticMeshSection.NumTriangles = 0;
