Fix the reduction over itself (lod 0 reduce)

diff --git a/Engine/Source/Developer/MeshBuilder/Private/StaticMeshBuilder.cpp b/Engine/Source/Developer/MeshBuilder/Private/StaticMeshBuilder.cpp
--- a/Engine/Source/Developer/MeshBuilder/Private/StaticMeshBuilder.cpp
+++ b/Engine/Source/Developer/MeshBuilder/Private/StaticMeshBuilder.cpp
@@ -114,8 +114,14 @@ bool FStaticMeshBuilder::Build(FStaticMeshRenderData& StaticMeshRenderData, USta
 			float OverlappingThreshold = LODBuildSettings.bRemoveDegenerates ? THRESH_POINTS_ARE_SAME : 0.0f;
 			TMultiMap<int32, int32> OverlappingCorners;
 			FMeshDescriptionOperations::FindOverlappingCorners(OverlappingCorners, StaticMesh->GetMeshDescription(BaseLODIndex), OverlappingThreshold);
-			MeshDescriptionHelper.ReduceLOD(StaticMesh->GetMeshDescription(BaseLODIndex), StaticMesh->GetMeshDescription(LodIndex), ReductionSettings, OverlappingCorners);
-			UMeshDescription* MeshDescriptionReduced = StaticMesh->GetMeshDescription(LodIndex);
+
+			//Create a new destination mesh in case we reduce ourself
+			UMeshDescription* MeshDescriptionReduced = NewObject<UMeshDescription>(StaticMesh, NAME_None);
+			UStaticMesh::RegisterMeshAttributes(MeshDescriptionReduced);
+			MeshDescriptionHelper.ReduceLOD(StaticMesh->GetMeshDescription(BaseLODIndex), MeshDescriptionReduced, ReductionSettings, OverlappingCorners);
+			StaticMesh->SetMeshDescription(LodIndex, MeshDescriptionReduced);
+			MeshDescription = MeshDescriptionReduced;
+
 			if (MeshDescriptionReduced != nullptr)
 			{
 				const TPolygonGroupAttributeArray<FName>& PolygonGroupImportedMaterialSlotNames = MeshDescriptionReduced->PolygonGroupAttributes().GetAttributes<FName>(MeshAttribute::PolygonGroup::ImportedMaterialSlotName);
diff --git a/Engine/Source/Runtime/Engine/Private/StaticMesh.cpp b/Engine/Source/Runtime/Engine/Private/StaticMesh.cpp
--- a/Engine/Source/Runtime/Engine/Private/StaticMesh.cpp
+++ b/Engine/Source/Runtime/Engine/Private/StaticMesh.cpp
@@ -1333,7 +1333,7 @@ FArchive& operator<<(FArchive& Ar, FMeshBuildSettings& BuildSettings)
 // differences, etc.) replace the version GUID below with a new one.
 // In case of merge conflicts with DDC versions, you MUST generate a new GUID
 // and set this new GUID as the version.                                       
-#define STATICMESH_DERIVEDDATA_VER TEXT("AEC7F25C2C9141A99EC26695EE5BA392")
+#define STATICMESH_DERIVEDDATA_VER TEXT("305A71049F594F5A8D73BE076F3DEF11")
 
 static const FString& GetStaticMeshDerivedDataVersion()
 {
@@ -2689,7 +2689,7 @@ void UStaticMesh::FixupMaterialSlotName()
 // differences, etc.) replace the version GUID below with a new one.
 // In case of merge conflicts with DDC versions, you MUST generate a new GUID
 // and set this new GUID as the version.                                       
-#define MESHDATAKEY_STATICMESH_DERIVEDDATA_VER TEXT("8BB01BA6862D43818AE1E78DA6BF67EC")
+#define MESHDATAKEY_STATICMESH_DERIVEDDATA_VER TEXT("0329735AAE644B4AAC069ED412F0574B")
 
 static const FString& GetMeshDataKeyStaticMeshDerivedDataVersion()
 {
