Fix automation tests warning, we always make sure there is no degenerate NTBs

diff --git a/Engine/Source/Developer/MeshBuilder/Private/MeshDescriptionHelper.cpp b/Engine/Source/Developer/MeshBuilder/Private/MeshDescriptionHelper.cpp
--- a/Engine/Source/Developer/MeshBuilder/Private/MeshDescriptionHelper.cpp
+++ b/Engine/Source/Developer/MeshBuilder/Private/MeshDescriptionHelper.cpp
@@ -106,10 +106,9 @@ UMeshDescription* FMeshDescriptionHelper::GetRenderMeshDescription(UObject* Owne
 			}
 			FMeshDescriptionOperations::CreateMikktTangents(RenderMeshDescription, (FMeshDescriptionOperations::ETangentOptions)TangentOptions);
 		}
-		else if(!bHasAllNormals || !bHasAllTangents)
+		else
 		{
 			FMeshDescriptionOperations::CreateNormals(RenderMeshDescription, (FMeshDescriptionOperations::ETangentOptions)TangentOptions, true);
-
 			//EComputeNTBsOptions ComputeNTBsOptions = (bHasAllNormals ? EComputeNTBsOptions::None : EComputeNTBsOptions::Normals) | (bHasAllTangents ? EComputeNTBsOptions::None : EComputeNTBsOptions::Tangents);
 			//RenderMeshDescription->ComputeTangentsAndNormals(ComputeNTBsOptions);
 		}
diff --git a/Engine/Source/Developer/MeshBuilder/Private/StaticMeshBuilder.cpp b/Engine/Source/Developer/MeshBuilder/Private/StaticMeshBuilder.cpp
--- a/Engine/Source/Developer/MeshBuilder/Private/StaticMeshBuilder.cpp
+++ b/Engine/Source/Developer/MeshBuilder/Private/StaticMeshBuilder.cpp
@@ -327,118 +327,123 @@ void BuildVertexBuffer(
 	const TVertexInstanceAttributeArray<FVector4>& VertexInstanceColors = MeshDescription->VertexInstanceAttributes().GetAttributes<FVector4>( MeshAttribute::VertexInstance::Color );
 	const TVertexInstanceAttributeIndicesArray<FVector2D>& VertexInstanceUVs = MeshDescription->VertexInstanceAttributes().GetAttributesSet<FVector2D>( MeshAttribute::VertexInstance::TextureCoordinate );
 
-	// Set up index buffer
+	TMap<FPolygonGroupID, int32> PolygonGroupToSectionIndex;
+	
+	
 	for (const FPolygonGroupID PolygonGroupID : MeshDescription->PolygonGroups().GetElementIDs())
 	{
-		const TArray<FPolygonID>& Polygons = MeshDescription->GetPolygonGroupPolygons(PolygonGroupID);
-
-		// Create new rendering section
-		int32 SectionIndex = StaticMeshLOD.Sections.Add(FStaticMeshSection());
-		FStaticMeshSection& StaticMeshSection = StaticMeshLOD.Sections.Last();
-		TArray<uint32>& SectionIndices = OutPerSectionIndices[SectionIndex];
-
-		StaticMeshSection.FirstIndex = IndexBuffer.Num();
-		StaticMeshSection.NumTriangles = 0;
-		for (const FPolygonID& PolygonID : Polygons)
-		{
-			StaticMeshSection.NumTriangles += (uint32)MeshDescription->GetPolygonTriangles(PolygonID).Num();
-		}
+		int32& SectionIndex = PolygonGroupToSectionIndex.FindOrAdd(PolygonGroupID);
+		SectionIndex = StaticMeshLOD.Sections.Add(FStaticMeshSection());
+		FStaticMeshSection& StaticMeshSection = StaticMeshLOD.Sections[SectionIndex];
 		StaticMeshSection.MaterialIndex = StaticMesh->GetMaterialIndexFromImportedMaterialSlotName(PolygonGroupImportedMaterialSlotNames[PolygonGroupID]);
-		
 		if (StaticMeshSection.MaterialIndex == INDEX_NONE)
 		{
 			StaticMeshSection.MaterialIndex = PolygonGroupID.GetValue();
 		}
+	}
 
-		for (const FPolygonID PolygonID : Polygons)
+	for (const FPolygonID& PolygonID : MeshDescription->Polygons().GetElementIDs())
+	{
+		const FPolygonGroupID PolygonGroupID = MeshDescription->GetPolygonPolygonGroup(PolygonID);
+		const int32 SectionIndex = PolygonGroupToSectionIndex[PolygonGroupID];
+		TArray<uint32>& SectionIndices = OutPerSectionIndices[SectionIndex];
+
+		const TArray<FMeshTriangle>& PolygonTriangles = MeshDescription->GetPolygonTriangles(PolygonID);
+		int32 ReserveSize = IndexBuffer.Num() + PolygonTriangles.Num() * 3;
+		IndexBuffer.Reserve(ReserveSize);
+		uint32 MinIndex = TNumericLimits< uint32 >::Max();
+		uint32 MaxIndex = TNumericLimits< uint32 >::Min();
+		for (int32 TriangleIndex = 0; TriangleIndex < PolygonTriangles.Num(); ++TriangleIndex)
 		{
-			const TArray<FMeshTriangle>& PolygonTriangles = MeshDescription->GetPolygonTriangles(PolygonID);
-			int32 ReserveSize = IndexBuffer.Num() + PolygonTriangles.Num() * 3;
-			IndexBuffer.Reserve(ReserveSize);
-			uint32 MinIndex = TNumericLimits< uint32 >::Max();
-			uint32 MaxIndex = TNumericLimits< uint32 >::Min();
-			for (int32 TriangleIndex = 0; TriangleIndex < PolygonTriangles.Num(); ++TriangleIndex)
+			const FMeshTriangle& Triangle = PolygonTriangles[TriangleIndex];
+
+			FVector CornerPositions[3];
+			for (int32 TriVert = 0; TriVert < 3; ++TriVert)
 			{
-				const FMeshTriangle& Triangle = PolygonTriangles[TriangleIndex];
-				for (int32 TriVert = 0; TriVert < 3; ++TriVert)
-				{
-					const FVertexInstanceID VertexInstanceID = Triangle.GetVertexInstanceID(TriVert);
-					const int32 VertexInstanceValue = VertexInstanceID.GetValue();
+				const FVertexInstanceID VertexInstanceID = Triangle.GetVertexInstanceID(TriVert);
+				const FVertexID VertexID = MeshDescription->GetVertexInstanceVertex(VertexInstanceID);
+				CornerPositions[TriVert] = VertexPositions[VertexID];
+			}
+			FOverlappingThresholds OverlappingThresholds;
+			OverlappingThresholds.ThresholdPosition = VertexComparisonThreshold;
+			// Don't process degenerate triangles.
+			if (PointsEqual(CornerPositions[0], CornerPositions[1], OverlappingThresholds)
+				|| PointsEqual(CornerPositions[0], CornerPositions[2], OverlappingThresholds)
+				|| PointsEqual(CornerPositions[1], CornerPositions[2], OverlappingThresholds))
+			{
+				continue;
+			}
 
-					const FVertexID VertexID = MeshDescription->GetVertexInstanceVertex(VertexInstanceID);
+			for (int32 TriVert = 0; TriVert < 3; ++TriVert)
+			{
+				const FVertexInstanceID VertexInstanceID = Triangle.GetVertexInstanceID(TriVert);
+				const int32 VertexInstanceValue = VertexInstanceID.GetValue();
+				const FVector& VertexPosition = CornerPositions[TriVert];
+				const FVector& VertexInstanceNormal = VertexInstanceNormals[VertexInstanceID];
+				const FVector& VertexInstanceTangent = VertexInstanceTangents[VertexInstanceID];
+				const float VertexInstanceBinormalSign = VertexInstanceBinormalSigns[VertexInstanceID];
+				const FVector4& VertexInstanceColor = VertexInstanceColors[VertexInstanceID];
+
+				const FLinearColor LinearColor(VertexInstanceColor);
+				if (LinearColor != FLinearColor::White)
+				{
+					bHasColor = true;
+				}
 
-					const FVector& VertexPosition = VertexPositions[VertexID];
-					const FVector& VertexInstanceNormal = VertexInstanceNormals[VertexInstanceID];
-					const FVector& VertexInstanceTangent = VertexInstanceTangents[VertexInstanceID];
-					const float VertexInstanceBinormalSign = VertexInstanceBinormalSigns[VertexInstanceID];
-					const FVector4& VertexInstanceColor = VertexInstanceColors[VertexInstanceID];
+				FStaticMeshBuildVertex StaticMeshVertex;
 
-					const FLinearColor LinearColor(VertexInstanceColor);
-					if (LinearColor != FLinearColor::White)
+				StaticMeshVertex.Position = VertexPosition * LODBuildSettings.BuildScale3D;
+				const FMatrix ScaleMatrix = FScaleMatrix(LODBuildSettings.BuildScale3D).Inverse().GetTransposed();
+				StaticMeshVertex.TangentX = ScaleMatrix.TransformVector(VertexInstanceTangent).GetSafeNormal();
+				StaticMeshVertex.TangentY = ScaleMatrix.TransformVector(FVector::CrossProduct(VertexInstanceNormal, VertexInstanceTangent).GetSafeNormal() * VertexInstanceBinormalSign).GetSafeNormal();
+				StaticMeshVertex.TangentZ = ScaleMatrix.TransformVector(VertexInstanceNormal).GetSafeNormal();
+				StaticMeshVertex.Color = LinearColor.ToFColor(true);
+				const uint32 MaxNumTexCoords = FMath::Min<int32>(MAX_MESH_TEXTURE_COORDS_MD, MAX_STATIC_TEXCOORDS);
+				for (uint32 UVIndex = 0; UVIndex < MaxNumTexCoords; ++UVIndex)
+				{
+					if(UVIndex < NumTextureCoord)
 					{
-						bHasColor = true;
+						StaticMeshVertex.UVs[UVIndex] = VertexInstanceUVs.GetArrayForIndex(UVIndex)[VertexInstanceID];
 					}
-
-					FStaticMeshBuildVertex StaticMeshVertex;
-
-					StaticMeshVertex.Position = VertexPosition * LODBuildSettings.BuildScale3D;
-					const FMatrix ScaleMatrix = FScaleMatrix(LODBuildSettings.BuildScale3D).Inverse().GetTransposed();
-					StaticMeshVertex.TangentX = ScaleMatrix.TransformVector(VertexInstanceTangent).GetSafeNormal();
-					StaticMeshVertex.TangentY = ScaleMatrix.TransformVector(FVector::CrossProduct(VertexInstanceNormal, VertexInstanceTangent).GetSafeNormal() * VertexInstanceBinormalSign).GetSafeNormal();
-					StaticMeshVertex.TangentZ = ScaleMatrix.TransformVector(VertexInstanceNormal).GetSafeNormal();
-					StaticMeshVertex.Color = LinearColor.ToFColor(true);
-					const uint32 MaxNumTexCoords = FMath::Min<int32>(MAX_MESH_TEXTURE_COORDS_MD, MAX_STATIC_TEXCOORDS);
-					for (uint32 UVIndex = 0; UVIndex < MaxNumTexCoords; ++UVIndex)
+					else
 					{
-						if(UVIndex < NumTextureCoord)
-						{
-							StaticMeshVertex.UVs[UVIndex] = VertexInstanceUVs.GetArrayForIndex(UVIndex)[VertexInstanceID];
-						}
-						else
-						{
-							StaticMeshVertex.UVs[UVIndex] = FVector2D(0.0f, 0.0f);
-						}
+						StaticMeshVertex.UVs[UVIndex] = FVector2D(0.0f, 0.0f);
 					}
+				}
 					
 
-					//Never add duplicated vertex instance
-					DupVerts.Reset();
-					OverlappingCorners.MultiFind(VertexInstanceValue, DupVerts);
-					int32 Index = INDEX_NONE;
-					for (int32 k = 0; k < DupVerts.Num(); k++)
+				//Never add duplicated vertex instance
+				DupVerts.Reset();
+				OverlappingCorners.MultiFind(VertexInstanceValue, DupVerts);
+				DupVerts.Sort();
+				int32 Index = INDEX_NONE;
+				for (int32 k = 0; k < DupVerts.Num(); k++)
+				{
+					if (DupVerts[k] >= VertexInstanceValue)
 					{
-						int32 Location = RemapVerts.IsValidIndex(DupVerts[k]) ? RemapVerts[DupVerts[k]] : INDEX_NONE;
-						if (Location != INDEX_NONE && AreVerticesEqual(StaticMeshVertex, StaticMeshBuildVertices[Location], VertexComparisonThreshold))
-						{
-							Index = Location;
-							break;
-						}
+						break;
 					}
-					if (Index == INDEX_NONE)
+					int32 Location = RemapVerts.IsValidIndex(DupVerts[k]) ? RemapVerts[DupVerts[k]] : INDEX_NONE;
+					if (Location != INDEX_NONE && AreVerticesEqual(StaticMeshVertex, StaticMeshBuildVertices[Location], VertexComparisonThreshold))
 					{
-						Index = StaticMeshBuildVertices.Add(StaticMeshVertex);
+						Index = Location;
+						break;
 					}
-					RemapVerts[VertexInstanceValue] = Index;
-					const uint32 RenderingVertexIndex = RemapVerts[VertexInstanceValue];
-					IndexBuffer.Add(RenderingVertexIndex);
-					OutWedgeMap[VertexInstanceValue] = RenderingVertexIndex;
-					SectionIndices.Add(RenderingVertexIndex);
-					MinIndex = FMath::Min(MinIndex, RenderingVertexIndex);
-					MaxIndex = FMath::Max(MaxIndex, RenderingVertexIndex);
 				}
+				if (Index == INDEX_NONE)
+				{
+					Index = StaticMeshBuildVertices.Add(StaticMeshVertex);
+				}
+				RemapVerts[VertexInstanceValue] = Index;
+				const uint32 RenderingVertexIndex = RemapVerts[VertexInstanceValue];
+				IndexBuffer.Add(RenderingVertexIndex);
+				OutWedgeMap[VertexInstanceValue] = RenderingVertexIndex;
+				SectionIndices.Add(RenderingVertexIndex);
 			}
-			StaticMeshSection.MinVertexIndex = MinIndex;
-			StaticMeshSection.MaxVertexIndex = MaxIndex;
-		}
-
-		if (Polygons.Num() == 0)
-		{
-			// No triangles in this section
-			StaticMeshSection.MinVertexIndex = 0;
-			StaticMeshSection.MaxVertexIndex = 0;
 		}
 	}
 
+
 	//Optimize before setting the buffer
 	if (VertexInstances.Num() < 100000 * 3)
 	{
diff --git a/Engine/Source/Developer/MeshDescriptionOperations/Private/MeshDescriptionOperations.cpp b/Engine/Source/Developer/MeshDescriptionOperations/Private/MeshDescriptionOperations.cpp
--- a/Engine/Source/Developer/MeshDescriptionOperations/Private/MeshDescriptionOperations.cpp
+++ b/Engine/Source/Developer/MeshDescriptionOperations/Private/MeshDescriptionOperations.cpp
@@ -353,7 +353,7 @@ void FillMeshDescriptionVertexPositionNoDuplicate(const TArray<FVector> &RawMesh
 		}
 		TempRemapVertexPosition.FindOrAdd(Index_i) = VertexCount;
 		// only need to search forward, since we add pairs both ways
-		for (int32 j = i + 1; j < VertIndexAndZ.Num(); j++)
+/*		for (int32 j = i + 1; j < VertIndexAndZ.Num(); j++)
 		{
 			if (FMath::Abs(VertIndexAndZ[j].Z - VertIndexAndZ[i].Z) > SMALL_NUMBER)
 				break; // can't be any more dups
@@ -365,7 +365,7 @@ void FillMeshDescriptionVertexPositionNoDuplicate(const TArray<FVector> &RawMesh
 			{
 				TempRemapVertexPosition.FindOrAdd(VertIndexAndZ[j].Index) = VertexCount;
 			}
-		}
+		}*/
 		VertexCount++;
 	}
 
@@ -728,6 +728,18 @@ void FMeshDescriptionOperations::CreateNormals(UMeshDescription* MeshDescription
 							VertexInfo.VertexInstanceID = VertexInstanceID;
 							VertexInfo.UVs = VertexUVs[VertexInstanceID];
 							bPointHasAllTangents &= !VertexNormals[VertexInstanceID].IsNearlyZero() && !VertexTangents[VertexInstanceID].IsNearlyZero();
+							if (bPointHasAllTangents)
+							{
+								FVector TangentX = VertexTangents[VertexInstanceID].GetSafeNormal();
+								FVector TangentZ = VertexNormals[VertexInstanceID].GetSafeNormal();
+								FVector TangentY = (FVector::CrossProduct(TangentZ, TangentX).GetSafeNormal() * VertexBinormalSigns[VertexInstanceID]).GetSafeNormal();
+								if (TangentX.ContainsNaN() || TangentX.IsNearlyZero(SMALL_NUMBER) ||
+									TangentY.ContainsNaN() || TangentY.IsNearlyZero(SMALL_NUMBER) ||
+									TangentZ.ContainsNaN() || TangentZ.IsNearlyZero(SMALL_NUMBER))
+								{
+									bPointHasAllTangents = false;
+								}
+							}
 							break;
 						}
 					}
@@ -801,17 +813,35 @@ void FMeshDescriptionOperations::CreateNormals(UMeshDescription* MeshDescription
 			FVector GroupNormal(0.0f);
 			for (const FPolygonID PolygonID : Group)
 			{
+				FVector PolyNormal = PolygonNormals[PolygonID];
+				FVector PolyTangent = PolygonTangents[PolygonID];
+				FVector PolyBinormal = PolygonBinormals[PolygonID];
+				
 				ConsumedPolygon.Add(PolygonID);
 				VertexInstanceInGroup.Add(VertexInfoMap[PolygonID].VertexInstanceID);
-				GroupNormal += PolygonNormals[PolygonID];
+				if (!PolyNormal.IsNearlyZero(SMALL_NUMBER) && !PolyNormal.ContainsNaN())
+				{
+					GroupNormal += PolyNormal;
+				}
 				if (bComputeTangent)
 				{
 					const FVector2D UVs = VertexInfoMap[PolygonID].UVs;
 					bool CreateGroup = (!GroupTangent.Contains(UVs));
 					FVector& GroupTangentValue = GroupTangent.FindOrAdd(UVs);
 					FVector& GroupBiNormalValue = GroupBiNormal.FindOrAdd(UVs);
-					GroupTangentValue = CreateGroup ? PolygonTangents[PolygonID] : GroupTangentValue + PolygonTangents[PolygonID];
-					GroupBiNormalValue = CreateGroup ? PolygonBinormals[PolygonID] : GroupBiNormalValue + PolygonBinormals[PolygonID];
+					if (CreateGroup)
+					{
+						GroupTangentValue = FVector(0.0f);
+						GroupBiNormalValue = FVector(0.0f);
+					}
+					if (!PolyTangent.IsNearlyZero(SMALL_NUMBER) && !PolyTangent.ContainsNaN())
+					{
+						GroupTangentValue += PolyTangent;
+					}
+					if (!PolyBinormal.IsNearlyZero(SMALL_NUMBER) && !PolyBinormal.ContainsNaN())
+					{
+						GroupBiNormalValue += PolyBinormal;
+					}
 				}
 			}
 
