Use the generated LOD SectionInfoMap in case the data is valid. Otherwise use the base LOD SectionInfoMap. This ensure the build give the same result has the when the asset get loaded.

diff --git a/Engine/Source/Developer/MeshBuilder/Private/StaticMeshBuilder.cpp b/Engine/Source/Developer/MeshBuilder/Private/StaticMeshBuilder.cpp
--- a/Engine/Source/Developer/MeshBuilder/Private/StaticMeshBuilder.cpp
+++ b/Engine/Source/Developer/MeshBuilder/Private/StaticMeshBuilder.cpp
@@ -110,18 +110,32 @@ bool FStaticMeshBuilder::Build(FStaticMeshRenderData& StaticMeshRenderData, USta
 			{
 				//Duplicate the lodindex 0 we have a 100% reduction which is like a duplicate
 				MeshDescription = Cast<UMeshDescription>(StaticDuplicateObject(GetMeshDescription(BaseReduceLodIndex), StaticMesh, NAME_None, RF_NoFlags));
+				if (LodIndex > 0)
+				{
+					
+					//Make sure the SectionInfoMap is taken from the Base RawMesh
+					int32 SectionNumber = StaticMesh->OriginalSectionInfoMap.GetSectionNumber(BaseReduceLodIndex);
+					for (int32 SectionIndex = 0; SectionIndex < SectionNumber; ++SectionIndex)
+					{
+						//Keep the old data if its valid
+						bool bHasValidLODInfoMap = StaticMesh->SectionInfoMap.IsValidSection(LodIndex, SectionIndex);
+						//Section material index have to be remap with the ReductionSettings.BaseLODModel SectionInfoMap to create
+						//a valid new section info map for the reduced LOD.
+						if (!bHasValidLODInfoMap && StaticMesh->SectionInfoMap.IsValidSection(BaseReduceLodIndex, SectionIndex))
+						{
+							//Copy the BaseLODModel section info to the reduce LODIndex.
+							FMeshSectionInfo SectionInfo = StaticMesh->SectionInfoMap.Get(BaseReduceLodIndex, SectionIndex);
+							FMeshSectionInfo OriginalSectionInfo = StaticMesh->OriginalSectionInfoMap.Get(BaseReduceLodIndex, SectionIndex);
+							StaticMesh->SectionInfoMap.Set(LodIndex, SectionIndex, SectionInfo);
+							StaticMesh->OriginalSectionInfoMap.Set(LodIndex, SectionIndex, OriginalSectionInfo);
+						}
+					}
+				}
 			}
+
 			if (LodIndex > 0)
 			{
 				LODBuildSettings = StaticMesh->SourceModels[BaseReduceLodIndex].BuildSettings;
-				//Make sure the SectionInfoMap is taken from the Base RawMesh
-				int32 SectionNumber = StaticMesh->OriginalSectionInfoMap.GetSectionNumber(BaseReduceLodIndex);
-				for (int32 SectionIndex = 0; SectionIndex < SectionNumber; ++SectionIndex)
-				{
-					FMeshSectionInfo Info = StaticMesh->OriginalSectionInfoMap.Get(BaseReduceLodIndex, SectionIndex);
-					StaticMesh->SectionInfoMap.Set(LodIndex, SectionIndex, Info);
-					StaticMesh->OriginalSectionInfoMap.Set(LodIndex, SectionIndex, Info);
-				}
 			}
 		}
 
@@ -154,6 +168,7 @@ bool FStaticMeshBuilder::Build(FStaticMeshRenderData& StaticMeshRenderData, USta
 			
 				//Set the new SectionInfoMap for this reduced LOD base on the ReductionSettings.BaseLODModel SectionInfoMap
 				const FMeshSectionInfoMap& LODModelSectionInfoMap = StaticMesh->SectionInfoMap;
+				const FMeshSectionInfoMap& LODModelOriginalSectionInfoMap = StaticMesh->OriginalSectionInfoMap;
 				TArray<int32> UniqueMaterialIndex;
 				//Find all unique Material in used order
 				for (const FPolygonGroupID& PolygonGroupID : MeshDescriptionReduced->PolygonGroups().GetElementIDs())
@@ -168,24 +183,17 @@ bool FStaticMeshBuilder::Build(FStaticMeshRenderData& StaticMeshRenderData, USta
 				//All used material represent a different section
 				for (int32 SectionIndex = 0; SectionIndex < UniqueMaterialIndex.Num(); ++SectionIndex)
 				{
+					//Keep the old data
+					bool bHasValidLODInfoMap = LODModelSectionInfoMap.IsValidSection(LodIndex, SectionIndex);
 					//Section material index have to be remap with the ReductionSettings.BaseLODModel SectionInfoMap to create
 					//a valid new section info map for the reduced LOD.
-					if (LODModelSectionInfoMap.IsValidSection(ReductionSettings.BaseLODModel, UniqueMaterialIndex[SectionIndex]))
+					if (!bHasValidLODInfoMap && LODModelSectionInfoMap.IsValidSection(ReductionSettings.BaseLODModel, UniqueMaterialIndex[SectionIndex]))
 					{
-						FMeshSectionInfo SectionInfo = LODModelSectionInfoMap.Get(ReductionSettings.BaseLODModel, UniqueMaterialIndex[SectionIndex]);
-						//Try to recuperate the valid data
-						if (LODModelSectionInfoMap.IsValidSection(LodIndex, SectionIndex))
-						{
-							//If the old LOD section was using the same Material copy the data
-							FMeshSectionInfo OriginalLODSectionInfo = LODModelSectionInfoMap.Get(LodIndex, SectionIndex);
-							if (OriginalLODSectionInfo.MaterialIndex == SectionInfo.MaterialIndex)
-							{
-								SectionInfo.bCastShadow = OriginalLODSectionInfo.bCastShadow;
-								SectionInfo.bEnableCollision = OriginalLODSectionInfo.bEnableCollision;
-							}
-						}
 						//Copy the BaseLODModel section info to the reduce LODIndex.
+						FMeshSectionInfo SectionInfo = LODModelSectionInfoMap.Get(ReductionSettings.BaseLODModel, UniqueMaterialIndex[SectionIndex]);
+						FMeshSectionInfo OriginalSectionInfo = LODModelOriginalSectionInfoMap.Get(ReductionSettings.BaseLODModel, UniqueMaterialIndex[SectionIndex]);
 						StaticMesh->SectionInfoMap.Set(LodIndex, SectionIndex, SectionInfo);
+						StaticMesh->OriginalSectionInfoMap.Set(LodIndex, SectionIndex, OriginalSectionInfo);
 					}
 				}
 			}
