Update the staticmesh LOD model max deviation when generating a LOD

diff --git a/Engine/Source/Developer/MeshBuilder/Private/MeshDescriptionHelper.cpp b/Engine/Source/Developer/MeshBuilder/Private/MeshDescriptionHelper.cpp
--- a/Engine/Source/Developer/MeshBuilder/Private/MeshDescriptionHelper.cpp
+++ b/Engine/Source/Developer/MeshBuilder/Private/MeshDescriptionHelper.cpp
@@ -147,7 +147,7 @@ UMeshDescription* FMeshDescriptionHelper::GetRenderMeshDescription(UObject* Owne
 	return RenderMeshDescription;
 }
 
-void FMeshDescriptionHelper::ReduceLOD(const UMeshDescription* BaseMesh, UMeshDescription* DestMesh, const FMeshReductionSettings& ReductionSettings, const TMultiMap<int32, int32>& InOverlappingCorners)
+void FMeshDescriptionHelper::ReduceLOD(const UMeshDescription* BaseMesh, UMeshDescription* DestMesh, const FMeshReductionSettings& ReductionSettings, const TMultiMap<int32, int32>& InOverlappingCorners, float &OutMaxDeviation)
 {
 	if (BaseMesh == nullptr || DestMesh == nullptr)
 	{
@@ -163,8 +163,8 @@ void FMeshDescriptionHelper::ReduceLOD(const UMeshDescription* BaseMesh, UMeshDe
 	{
 		return;
 	}
-	float MaxDeviation = ReductionSettings.MaxDeviation;
-	MeshReduction->ReduceMeshDescription(DestMesh, MaxDeviation, BaseMesh, InOverlappingCorners, ReductionSettings);
+	OutMaxDeviation = ReductionSettings.MaxDeviation;
+	MeshReduction->ReduceMeshDescription(DestMesh, OutMaxDeviation, BaseMesh, InOverlappingCorners, ReductionSettings);
 }
 
 bool FMeshDescriptionHelper::IsValidOriginalMeshDescription()
diff --git a/Engine/Source/Developer/MeshBuilder/Private/MeshDescriptionHelper.h b/Engine/Source/Developer/MeshBuilder/Private/MeshDescriptionHelper.h
--- a/Engine/Source/Developer/MeshBuilder/Private/MeshDescriptionHelper.h
+++ b/Engine/Source/Developer/MeshBuilder/Private/MeshDescriptionHelper.h
@@ -27,7 +27,7 @@ public:
 	//Build a render mesh description with the BuildSettings. This will update the InRenderMeshDescription ptr content
 	UMeshDescription* GetRenderMeshDescription(UObject* Owner);
 
-	void ReduceLOD(const UMeshDescription* BaseMesh, UMeshDescription* DestMesh, const struct FMeshReductionSettings& ReductionSettings, const TMultiMap<int32, int32>& InOverlappingCorners);
+	void ReduceLOD(const UMeshDescription* BaseMesh, UMeshDescription* DestMesh, const struct FMeshReductionSettings& ReductionSettings, const TMultiMap<int32, int32>& InOverlappingCorners, float &OutMaxDeviation);
 
 	//Return true if there is a valid original mesh description, false otherwise(Auto generate LOD).
 	bool IsValidOriginalMeshDescription();
diff --git a/Engine/Source/Developer/MeshBuilder/Private/StaticMeshBuilder.cpp b/Engine/Source/Developer/MeshBuilder/Private/StaticMeshBuilder.cpp
--- a/Engine/Source/Developer/MeshBuilder/Private/StaticMeshBuilder.cpp
+++ b/Engine/Source/Developer/MeshBuilder/Private/StaticMeshBuilder.cpp
@@ -80,6 +80,7 @@ bool FStaticMeshBuilder::Build(FStaticMeshRenderData& StaticMeshRenderData, USta
 
 	for (int32 LodIndex = 0; LodIndex < StaticMesh->SourceModels.Num(); ++LodIndex)
 	{
+		float MaxDeviation = 0.0f;
 		FMeshBuildSettings& LODBuildSettings = StaticMesh->SourceModels[LodIndex].BuildSettings;
 		UMeshDescription* OriginalMeshDescription = StaticMesh->GetOriginalMeshDescription(LodIndex);
 		FMeshDescriptionHelper MeshDescriptionHelper(&LODBuildSettings, OriginalMeshDescription);
@@ -153,7 +154,7 @@ bool FStaticMeshBuilder::Build(FStaticMeshRenderData& StaticMeshRenderData, USta
 			//Create a new destination mesh in case we reduce ourself
 			UMeshDescription* MeshDescriptionReduced = NewObject<UMeshDescription>(StaticMesh, NAME_None);
 			UStaticMesh::RegisterMeshAttributes(MeshDescriptionReduced);
-			MeshDescriptionHelper.ReduceLOD(GetMeshDescription(BaseLODIndex), MeshDescriptionReduced, ReductionSettings, OverlappingCorners);
+			MeshDescriptionHelper.ReduceLOD(GetMeshDescription(BaseLODIndex), MeshDescriptionReduced, ReductionSettings, OverlappingCorners, MaxDeviation);
 			SetMeshDescription(LodIndex, MeshDescriptionReduced);
 			MeshDescription = MeshDescriptionReduced;
 
@@ -202,6 +203,7 @@ bool FStaticMeshBuilder::Build(FStaticMeshRenderData& StaticMeshRenderData, USta
 		const FPolygonGroupArray& PolygonGroups = MeshDescription->PolygonGroups();
 
 		FStaticMeshLODResources& StaticMeshLOD = StaticMeshRenderData.LODResources[LodIndex];
+		StaticMeshLOD.MaxDeviation = MaxDeviation;
 
 		//discover degenerate triangle with this threshold
 		float VertexComparisonThreshold = LODBuildSettings.bRemoveDegenerates ? THRESH_POINTS_ARE_SAME : 0.0f;
