Allow StaticMesh LOD 0 to be reduce

diff --git a/Engine/Source/Developer/MeshBuilder/Private/StaticMeshBuilder.cpp b/Engine/Source/Developer/MeshBuilder/Private/StaticMeshBuilder.cpp
--- a/Engine/Source/Developer/MeshBuilder/Private/StaticMeshBuilder.cpp
+++ b/Engine/Source/Developer/MeshBuilder/Private/StaticMeshBuilder.cpp
@@ -73,7 +73,7 @@ bool FStaticMeshBuilder::Build(FStaticMeshRenderData& StaticMeshRenderData, USta
 		const FStaticMeshSourceModel& SrcModel = StaticMesh->SourceModels[LodIndex];
 		FMeshReductionSettings ReductionSettings = LODGroup.GetSettings(SrcModel.ReductionSettings, LodIndex);
 
-		bool bUseReduction = (ReductionSettings.PercentTriangles < 1.0f || ReductionSettings.MaxDeviation > 0.0f) && LodIndex > 0;
+		bool bUseReduction = (ReductionSettings.PercentTriangles < 1.0f || ReductionSettings.MaxDeviation > 0.0f);
 
 		if (OriginalMeshDescription != nullptr)
 		{
@@ -130,7 +130,18 @@ bool FStaticMeshBuilder::Build(FStaticMeshRenderData& StaticMeshRenderData, USta
 
 			//Create a reduced mesh from the base LOD
 			UStaticMesh::RegisterMeshAttributes(MeshDescriptions[LodIndex]);
-			MeshDescriptionHelper.ReduceLOD(MeshDescriptions[BaseLodIndex], MeshDescriptions[LodIndex], ReductionSettings, OverlappingCorners, MaxDeviation);
+			
+			if (LodIndex == BaseLodIndex)
+			{
+				//When using LOD 0, we use a copy of the mesh description since reduce do not support inline reducing
+				FMeshDescription BaseMeshDescription = MeshDescriptions[BaseLodIndex];
+				MeshDescriptionHelper.ReduceLOD(BaseMeshDescription, MeshDescriptions[LodIndex], ReductionSettings, OverlappingCorners, MaxDeviation);
+			}
+			else
+			{
+				MeshDescriptionHelper.ReduceLOD(MeshDescriptions[BaseLodIndex], MeshDescriptions[LodIndex], ReductionSettings, OverlappingCorners, MaxDeviation);
+			}
+			
 
 			const TPolygonGroupAttributeArray<FName>& PolygonGroupImportedMaterialSlotNames = MeshDescriptions[LodIndex].PolygonGroupAttributes().GetAttributes<FName>(MeshAttribute::PolygonGroup::ImportedMaterialSlotName);
 			// Recompute adjacency information. Since we change the vertices when we reduce
