Fix the reduce to 100%, the overlapp vertex was not set in this case

diff --git a/Engine/Source/Developer/MeshBuilder/Private/StaticMeshBuilder.cpp b/Engine/Source/Developer/MeshBuilder/Private/StaticMeshBuilder.cpp
--- a/Engine/Source/Developer/MeshBuilder/Private/StaticMeshBuilder.cpp
+++ b/Engine/Source/Developer/MeshBuilder/Private/StaticMeshBuilder.cpp
@@ -65,6 +65,7 @@ bool FStaticMeshBuilder::Build(FStaticMeshRenderData& StaticMeshRenderData, USta
 
 	for (int32 LodIndex = 0; LodIndex < StaticMesh->SourceModels.Num(); ++LodIndex)
 	{
+		const int32 BaseReduceLodIndex = 0;
 		float MaxDeviation = 0.0f;
 		FMeshBuildSettings& LODBuildSettings = StaticMesh->SourceModels[LodIndex].BuildSettings;
 		const FMeshDescription* OriginalMeshDescription = StaticMesh->GetOriginalMeshDescription(LodIndex);
@@ -81,7 +82,6 @@ bool FStaticMeshBuilder::Build(FStaticMeshRenderData& StaticMeshRenderData, USta
 		}
 		else
 		{
-			const int32 BaseReduceLodIndex = 0;
 			if (bUseReduction)
 			{
 				// Initialize an empty mesh description that the reduce will fill
@@ -91,6 +91,9 @@ bool FStaticMeshBuilder::Build(FStaticMeshRenderData& StaticMeshRenderData, USta
 			{
 				//Duplicate the lodindex 0 we have a 100% reduction which is like a duplicate
 				MeshDescriptions[LodIndex] = MeshDescriptions[BaseReduceLodIndex];
+				//Set the overlapping threshold
+				float ComparisonThreshold = StaticMesh->SourceModels[BaseReduceLodIndex].BuildSettings.bRemoveDegenerates ? THRESH_POINTS_ARE_SAME : 0.0f;
+				MeshDescriptionHelper.FindOverlappingCorners(MeshDescriptions[LodIndex], ComparisonThreshold);
 				if (LodIndex > 0)
 				{
 					
diff --git a/Engine/Source/Runtime/Engine/Private/StaticMesh.cpp b/Engine/Source/Runtime/Engine/Private/StaticMesh.cpp
--- a/Engine/Source/Runtime/Engine/Private/StaticMesh.cpp
+++ b/Engine/Source/Runtime/Engine/Private/StaticMesh.cpp
@@ -1350,7 +1350,7 @@ FArchive& operator<<(FArchive& Ar, FMeshBuildSettings& BuildSettings)
 // differences, etc.) replace the version GUID below with a new one.
 // In case of merge conflicts with DDC versions, you MUST generate a new GUID
 // and set this new GUID as the version.                                       
-#define STATICMESH_DERIVEDDATA_VER TEXT("4FDA6C9D166447EAB2DA9EDA8A184591")
+#define STATICMESH_DERIVEDDATA_VER TEXT("3713973CA1B84F41BA1EB2E56FCE9211")
 
 static const FString& GetStaticMeshDerivedDataVersion()
 {
