Update the staticmesh LOD model max deviation when generating a LOD

diff --git a/Engine/Source/Developer/MeshBuilder/Private/MeshDescriptionHelper.cpp b/Engine/Source/Developer/MeshBuilder/Private/MeshDescriptionHelper.cpp
--- a/Engine/Source/Developer/MeshBuilder/Private/MeshDescriptionHelper.cpp
+++ b/Engine/Source/Developer/MeshBuilder/Private/MeshDescriptionHelper.cpp
@@ -137,7 +137,7 @@ void FMeshDescriptionHelper::GetRenderMeshDescription(UObject* Owner, const FMes
 	}
 }
 
-void FMeshDescriptionHelper::ReduceLOD(const FMeshDescription& BaseMesh, FMeshDescription& DestMesh, const FMeshReductionSettings& ReductionSettings, const TMultiMap<int32, int32>& InOverlappingCorners)
+void FMeshDescriptionHelper::ReduceLOD(const FMeshDescription& BaseMesh, FMeshDescription& DestMesh, const FMeshReductionSettings& ReductionSettings, const TMultiMap<int32, int32>& InOverlappingCorners, float &OutMaxDeviation)
 {
 	IMeshReductionManagerModule& MeshReductionModule = FModuleManager::Get().LoadModuleChecked<IMeshReductionManagerModule>("MeshReductionInterface");
 	IMeshReduction* MeshReduction = MeshReductionModule.GetStaticMeshReductionInterface();
@@ -147,8 +147,8 @@ void FMeshDescriptionHelper::ReduceLOD(const FMeshDescription& BaseMesh, FMeshDe
 	{
 		return;
 	}
-	float MaxDeviation = ReductionSettings.MaxDeviation;
-	MeshReduction->ReduceMeshDescription(DestMesh, MaxDeviation, BaseMesh, InOverlappingCorners, ReductionSettings);
+	OutMaxDeviation = ReductionSettings.MaxDeviation;
+	MeshReduction->ReduceMeshDescription(DestMesh, OutMaxDeviation, BaseMesh, InOverlappingCorners, ReductionSettings);
 }
 
 void FMeshDescriptionHelper::FindOverlappingCorners(const FMeshDescription& MeshDescription, float ComparisonThreshold)
diff --git a/Engine/Source/Developer/MeshBuilder/Private/MeshDescriptionHelper.h b/Engine/Source/Developer/MeshBuilder/Private/MeshDescriptionHelper.h
--- a/Engine/Source/Developer/MeshBuilder/Private/MeshDescriptionHelper.h
+++ b/Engine/Source/Developer/MeshBuilder/Private/MeshDescriptionHelper.h
@@ -28,7 +28,7 @@ public:
 	//Build a render mesh description with the BuildSettings. This will update the InRenderMeshDescription ptr content
 	void GetRenderMeshDescription(UObject* Owner, const FMeshDescription& InOriginalMeshDescription, FMeshDescription& OutRenderMeshDescription);
 
-	void ReduceLOD(const FMeshDescription& BaseMesh, FMeshDescription& DestMesh, const FMeshReductionSettings& ReductionSettings, const TMultiMap<int32, int32>& InOverlappingCorners);
+	void ReduceLOD(const FMeshDescription& BaseMesh, FMeshDescription& DestMesh, const FMeshReductionSettings& ReductionSettings, const TMultiMap<int32, int32>& InOverlappingCorners, float &OutMaxDeviation);
 
 	void FindOverlappingCorners(const FMeshDescription& MeshDescription, float ComparisonThreshold);
 
diff --git a/Engine/Source/Developer/MeshBuilder/Private/StaticMeshBuilder.cpp b/Engine/Source/Developer/MeshBuilder/Private/StaticMeshBuilder.cpp
--- a/Engine/Source/Developer/MeshBuilder/Private/StaticMeshBuilder.cpp
+++ b/Engine/Source/Developer/MeshBuilder/Private/StaticMeshBuilder.cpp
@@ -65,6 +65,7 @@ bool FStaticMeshBuilder::Build(FStaticMeshRenderData& StaticMeshRenderData, USta
 
 	for (int32 LodIndex = 0; LodIndex < StaticMesh->SourceModels.Num(); ++LodIndex)
 	{
+		float MaxDeviation = 0.0f;
 		FMeshBuildSettings& LODBuildSettings = StaticMesh->SourceModels[LodIndex].BuildSettings;
 		const FMeshDescription* OriginalMeshDescription = StaticMesh->GetOriginalMeshDescription(LodIndex);
 		FMeshDescriptionHelper MeshDescriptionHelper(&LODBuildSettings);
@@ -116,7 +117,7 @@ bool FStaticMeshBuilder::Build(FStaticMeshRenderData& StaticMeshRenderData, USta
 
 			//Create a reduced mesh from the base LOD
 			UStaticMesh::RegisterMeshAttributes(MeshDescriptions[LodIndex]);
-			MeshDescriptionHelper.ReduceLOD(MeshDescriptions[BaseLodIndex], MeshDescriptions[LodIndex], ReductionSettings, OverlappingCorners);
+			MeshDescriptionHelper.ReduceLOD(MeshDescriptions[BaseLodIndex], MeshDescriptions[LodIndex], ReductionSettings, OverlappingCorners, MaxDeviation);
 
 			const TPolygonGroupAttributeArray<FName>& PolygonGroupImportedMaterialSlotNames = MeshDescriptions[LodIndex].PolygonGroupAttributes().GetAttributes<FName>(MeshAttribute::PolygonGroup::ImportedMaterialSlotName);
 			// Recompute adjacency information. Since we change the vertices when we reduce
@@ -166,6 +167,7 @@ bool FStaticMeshBuilder::Build(FStaticMeshRenderData& StaticMeshRenderData, USta
 		const FPolygonGroupArray& PolygonGroups = MeshDescriptions[LodIndex].PolygonGroups();
 
 		FStaticMeshLODResources& StaticMeshLOD = StaticMeshRenderData.LODResources[LodIndex];
+		StaticMeshLOD.MaxDeviation = MaxDeviation;
 
 		//discover degenerate triangle with this threshold
 		float VertexComparisonThreshold = LODBuildSettings.bRemoveDegenerates ? THRESH_POINTS_ARE_SAME : 0.0f;
