Fix staticmesh builder merge error

diff --git a/Engine/Source/Developer/MeshBuilder/Private/StaticMeshBuilder.cpp b/Engine/Source/Developer/MeshBuilder/Private/StaticMeshBuilder.cpp
--- a/Engine/Source/Developer/MeshBuilder/Private/StaticMeshBuilder.cpp
+++ b/Engine/Source/Developer/MeshBuilder/Private/StaticMeshBuilder.cpp
@@ -173,19 +173,8 @@ bool FStaticMeshBuilder::Build(FStaticMeshRenderData& StaticMeshRenderData, USta
 				//a valid new section info map for the reduced LOD.
 				if (!bHasValidLODInfoMap && LODModelSectionInfoMap.IsValidSection(ReductionSettings.BaseLODModel, UniqueMaterialIndex[SectionIndex]))
 				{
-					FMeshSectionInfo SectionInfo = LODModelSectionInfoMap.Get(ReductionSettings.BaseLODModel, UniqueMaterialIndex[SectionIndex]);
-					//Try to recuperate the valid data
-					if (LODModelSectionInfoMap.IsValidSection(LodIndex, SectionIndex))
-					{
-						//If the old LOD section was using the same Material copy the data
-						FMeshSectionInfo OriginalLODSectionInfo = LODModelSectionInfoMap.Get(LodIndex, SectionIndex);
-						if (OriginalLODSectionInfo.MaterialIndex == SectionInfo.MaterialIndex)
-						{
-							SectionInfo.bCastShadow = OriginalLODSectionInfo.bCastShadow;
-							SectionInfo.bEnableCollision = OriginalLODSectionInfo.bEnableCollision;
-						}
-					}
 					//Copy the BaseLODModel section info to the reduce LODIndex.
+					FMeshSectionInfo SectionInfo = LODModelSectionInfoMap.Get(ReductionSettings.BaseLODModel, UniqueMaterialIndex[SectionIndex]);
 					FMeshSectionInfo OriginalSectionInfo = LODModelOriginalSectionInfoMap.Get(ReductionSettings.BaseLODModel, UniqueMaterialIndex[SectionIndex]);
 					StaticMesh->SectionInfoMap.Set(LodIndex, SectionIndex, SectionInfo);
 					StaticMesh->OriginalSectionInfoMap.Set(LodIndex, SectionIndex, OriginalSectionInfo);
