Copy CL #4392021
  - Fixed crash when importing model thru VRED importer
  - Includes changes to StaticMeshBuilder.cpp made in CL #4369604

diff --git a/Engine/Source/Developer/MeshBuilder/Private/StaticMeshBuilder.cpp b/Engine/Source/Developer/MeshBuilder/Private/StaticMeshBuilder.cpp
--- a/Engine/Source/Developer/MeshBuilder/Private/StaticMeshBuilder.cpp
+++ b/Engine/Source/Developer/MeshBuilder/Private/StaticMeshBuilder.cpp
@@ -317,19 +317,10 @@ void BuildVertexBuffer(
 	const FPolygonGroupArray& PolygonGroupArray = MeshDescription.PolygonGroups();
 	const FPolygonArray& PolygonArray = MeshDescription.Polygons();
 	
-	OutWedgeMap.Reset();
-	OutWedgeMap.AddZeroed(VertexInstances.GetArraySize());
-
 	TArray<int32> RemapVertexInstanceID;
 	// set up vertex buffer elements
 	StaticMeshBuildVertices.Reserve(VertexInstances.GetArraySize());
 	bool bHasColor = false;
-	//Fill the remap array
-	RemapVerts.AddZeroed(VertexInstances.GetArraySize());
-	for (int32& RemapIndex : RemapVerts)
-	{
-		RemapIndex = INDEX_NONE;
-	}
 
 	TPolygonGroupAttributesConstRef<FName> PolygonGroupImportedMaterialSlotNames = MeshDescription.PolygonGroupAttributes().GetAttributesRef<FName>(MeshAttribute::PolygonGroup::ImportedMaterialSlotName);
 	TVertexAttributesConstRef<FVector> VertexPositions = MeshDescription.VertexAttributes().GetAttributesRef<FVector>( MeshAttribute::Vertex::Position );
@@ -364,6 +355,18 @@ void BuildVertexBuffer(
 	}
 	IndexBuffer.Reset(ReserveIndicesCount);
 
+	//Fill the remap array
+	RemapVerts.AddZeroed(ReserveIndicesCount);
+	for (int32& RemapIndex : RemapVerts)
+	{
+		RemapIndex = INDEX_NONE;
+	}
+
+	//Initialize the wedge map array tracking correspondence between wedge index and rendering vertex index
+	OutWedgeMap.Reset();
+	OutWedgeMap.AddZeroed(ReserveIndicesCount);
+
+	int32 WedgeIndex = 0;
 	for (const FPolygonID& PolygonID : MeshDescription.Polygons().GetElementIDs())
 	{
 		const FPolygonGroupID PolygonGroupID = MeshDescription.GetPolygonPolygonGroup(PolygonID);
@@ -394,7 +397,7 @@ void BuildVertexBuffer(
 				continue;
 			}
 
-			for (int32 TriVert = 0; TriVert < 3; ++TriVert)
+			for (int32 TriVert = 0; TriVert < 3; ++TriVert, ++WedgeIndex)
 			{
 				const FVertexInstanceID VertexInstanceID = Triangle.GetVertexInstanceID(TriVert);
 				const int32 VertexInstanceValue = VertexInstanceID.GetValue();
@@ -433,12 +436,13 @@ void BuildVertexBuffer(
 					
 
 				//Never add duplicated vertex instance
-				const TArray<int32>& DupVerts = OverlappingCorners.FindIfOverlapping(VertexInstanceValue);
+				//Use WedgeIndex since OverlappingCorners has been built based on that
+				const TArray<int32>& DupVerts = OverlappingCorners.FindIfOverlapping(WedgeIndex);
 
 				int32 Index = INDEX_NONE;
 				for (int32 k = 0; k < DupVerts.Num(); k++)
 				{
-					if (DupVerts[k] >= VertexInstanceValue)
+					if (DupVerts[k] >= WedgeIndex)
 					{
 						break;
 					}
@@ -453,10 +457,10 @@ void BuildVertexBuffer(
 				{
 					Index = StaticMeshBuildVertices.Add(StaticMeshVertex);
 				}
-				RemapVerts[VertexInstanceValue] = Index;
-				const uint32 RenderingVertexIndex = RemapVerts[VertexInstanceValue];
+				RemapVerts[WedgeIndex] = Index;
+				const uint32 RenderingVertexIndex = Index;
 				IndexBuffer.Add(RenderingVertexIndex);
-				OutWedgeMap[VertexInstanceValue] = RenderingVertexIndex;
+				OutWedgeMap[WedgeIndex] = RenderingVertexIndex;
 				SectionIndices.Add(RenderingVertexIndex);
 			}
 		}
diff --git a/Engine/Source/Runtime/MeshDescription/Public/MeshElementArray.h b/Engine/Source/Runtime/MeshDescription/Public/MeshElementArray.h
--- a/Engine/Source/Runtime/MeshDescription/Public/MeshElementArray.h
+++ b/Engine/Source/Runtime/MeshDescription/Public/MeshElementArray.h
@@ -259,8 +259,8 @@ public:
 template <typename ElementType>
 void TMeshElementArrayBase<ElementType>::Compact( TSparseArray<int32>& OutIndexRemap )
 {
-	static TSparseArray<ElementType> NewContainer;
-	NewContainer.Empty( Container.Num() );
+	TSparseArray<ElementType> NewContainer;
+	NewContainer.Reserve( Container.Num() );
 
 	OutIndexRemap.Empty( Container.GetMaxIndex() );
 
@@ -284,8 +284,8 @@ void TMeshElementArrayBase<ElementType>::Compact( TSparseArray<int32>& OutIndexR
 template <typename ElementType>
 void TMeshElementArrayBase<ElementType>::Remap( const TSparseArray<int32>& IndexRemap )
 {
-	static TSparseArray<ElementType> NewContainer;
-	NewContainer.Empty( IndexRemap.GetMaxIndex() );
+	TSparseArray<ElementType> NewContainer;
+	NewContainer.Reserve( IndexRemap.GetMaxIndex() );
 
 	// Add valid elements into a new contiguous sparse array.  Note non-const iterator so we can move elements.
 	for( typename TSparseArray<ElementType>::TIterator It( Container ); It; ++It )
