Port a fix from cl: 4366690
Fix the material ID when creating generating reduce LODs from a mesh where the materials was re-order to the FBX order.

diff --git a/Engine/Source/Developer/MeshBuilder/Private/StaticMeshBuilder.cpp b/Engine/Source/Developer/MeshBuilder/Private/StaticMeshBuilder.cpp
--- a/Engine/Source/Developer/MeshBuilder/Private/StaticMeshBuilder.cpp
+++ b/Engine/Source/Developer/MeshBuilder/Private/StaticMeshBuilder.cpp
@@ -174,13 +174,32 @@ bool FStaticMeshBuilder::Build(FStaticMeshRenderData& StaticMeshRenderData, USta
 				bool bHasValidLODInfoMap = LODModelSectionInfoMap.IsValidSection(LodIndex, SectionIndex);
 				//Section material index have to be remap with the ReductionSettings.BaseLODModel SectionInfoMap to create
 				//a valid new section info map for the reduced LOD.
-				if (!bHasValidLODInfoMap && LODModelSectionInfoMap.IsValidSection(ReductionSettings.BaseLODModel, UniqueMaterialIndex[SectionIndex]))
+
+				//Find the base LOD section using this material
+				if (!bHasValidLODInfoMap)
 				{
-					//Copy the BaseLODModel section info to the reduce LODIndex.
-					FMeshSectionInfo SectionInfo = LODModelSectionInfoMap.Get(ReductionSettings.BaseLODModel, UniqueMaterialIndex[SectionIndex]);
-					FMeshSectionInfo OriginalSectionInfo = LODModelOriginalSectionInfoMap.Get(ReductionSettings.BaseLODModel, UniqueMaterialIndex[SectionIndex]);
-					StaticMesh->SectionInfoMap.Set(LodIndex, SectionIndex, SectionInfo);
-					StaticMesh->OriginalSectionInfoMap.Set(LodIndex, SectionIndex, OriginalSectionInfo);
+					bool bSectionInfoSet = false;
+					for (int32 BaseLodSectionIndex = 0; BaseLodSectionIndex < LODModelSectionInfoMap.GetSectionNumber(BaseLodIndex); ++BaseLodSectionIndex)
+					{
+						FMeshSectionInfo SectionInfo = LODModelSectionInfoMap.Get(ReductionSettings.BaseLODModel, BaseLodSectionIndex);
+						if (SectionInfo.MaterialIndex == UniqueMaterialIndex[SectionIndex])
+						{
+							//Copy the BaseLODModel section info to the reduce LODIndex.
+							FMeshSectionInfo OriginalSectionInfo = LODModelOriginalSectionInfoMap.Get(ReductionSettings.BaseLODModel, BaseLodSectionIndex);
+							StaticMesh->SectionInfoMap.Set(LodIndex, SectionIndex, SectionInfo);
+							StaticMesh->OriginalSectionInfoMap.Set(LodIndex, SectionIndex, OriginalSectionInfo);
+							bSectionInfoSet = true;
+							break;
+						}
+					}
+					if (!bSectionInfoSet)
+					{
+						//Just set the default section info in case we did not found any match with the Base Lod
+						FMeshSectionInfo SectionInfo;
+						SectionInfo.MaterialIndex = SectionIndex;
+						StaticMesh->SectionInfoMap.Set(LodIndex, SectionIndex, SectionInfo);
+						StaticMesh->OriginalSectionInfoMap.Set(LodIndex, SectionIndex, SectionInfo);
+					}
 				}
 			}
 		}
