Make sure reduce LOD 0 do not loose there UVChannel

diff --git a/Engine/Source/Developer/MeshBuilder/Private/StaticMeshBuilder.cpp b/Engine/Source/Developer/MeshBuilder/Private/StaticMeshBuilder.cpp
--- a/Engine/Source/Developer/MeshBuilder/Private/StaticMeshBuilder.cpp
+++ b/Engine/Source/Developer/MeshBuilder/Private/StaticMeshBuilder.cpp
@@ -150,9 +150,6 @@ bool FStaticMeshBuilder::Build(FStaticMeshRenderData& StaticMeshRenderData, USta
 			FOverlappingCorners OverlappingCorners;
 			FMeshDescriptionOperations::FindOverlappingCorners(OverlappingCorners, MeshDescriptions[BaseReduceLodIndex], OverlappingThreshold);
 
-			//Create a reduced mesh from the base LOD
-			UStaticMesh::RegisterMeshAttributes(MeshDescriptions[LodIndex]);
-			
 			int32 OldSectionInfoMapCount = StaticMesh->SectionInfoMap.GetSectionNumber(LodIndex);
 
 			if (LodIndex == BaseReduceLodIndex)
diff --git a/Engine/Source/Runtime/Engine/Private/StaticMesh.cpp b/Engine/Source/Runtime/Engine/Private/StaticMesh.cpp
--- a/Engine/Source/Runtime/Engine/Private/StaticMesh.cpp
+++ b/Engine/Source/Runtime/Engine/Private/StaticMesh.cpp
@@ -1411,7 +1411,7 @@ static void SerializeBuildSettingsForDDC(FArchive& Ar, FMeshBuildSettings& Build
 // differences, etc.) replace the version GUID below with a new one.
 // In case of merge conflicts with DDC versions, you MUST generate a new GUID
 // and set this new GUID as the version.                                       
-#define STATICMESH_DERIVEDDATA_VER TEXT("7CF67E1C6DAD4EBE98D83C06C2FBC2CC")
+#define STATICMESH_DERIVEDDATA_VER TEXT("F36ADCAE95C14BAF91F7B137599924AB")
 
 static const FString& GetStaticMeshDerivedDataVersion()
 {
