Fix LOD Material assignment error when reimporting staticmesh and there is some LODs that are generated and the order of the SectionInfoMap is changed by old re-import.

diff --git a/Engine/Source/Developer/MeshBuilder/Private/StaticMeshBuilder.cpp b/Engine/Source/Developer/MeshBuilder/Private/StaticMeshBuilder.cpp
--- a/Engine/Source/Developer/MeshBuilder/Private/StaticMeshBuilder.cpp
+++ b/Engine/Source/Developer/MeshBuilder/Private/StaticMeshBuilder.cpp
@@ -81,6 +81,9 @@ bool FStaticMeshBuilder::Build(FStaticMeshRenderData& StaticMeshRenderData, USta
 	TArray<FMeshDescription> MeshDescriptions;
 	MeshDescriptions.SetNum(StaticMesh->SourceModels.Num());
 
+	const FMeshSectionInfoMap BeforeBuildSectionInfoMap = StaticMesh->SectionInfoMap;
+	const FMeshSectionInfoMap BeforeBuildOriginalSectionInfoMap = StaticMesh->OriginalSectionInfoMap;
+
 	for (int32 LodIndex = 0; LodIndex < StaticMesh->SourceModels.Num(); ++LodIndex)
 	{
 		float MaxDeviation = 0.0f;
@@ -165,15 +168,24 @@ bool FStaticMeshBuilder::Build(FStaticMeshRenderData& StaticMeshRenderData, USta
 			
 
 			const TPolygonGroupAttributesRef<FName> PolygonGroupImportedMaterialSlotNames = MeshDescriptions[LodIndex].PolygonGroupAttributes().GetAttributesRef<FName>(MeshAttribute::PolygonGroup::ImportedMaterialSlotName);
+			const TPolygonGroupAttributesRef<FName> BasePolygonGroupImportedMaterialSlotNames = MeshDescriptions[BaseReduceLodIndex].PolygonGroupAttributes().GetAttributesRef<FName>(MeshAttribute::PolygonGroup::ImportedMaterialSlotName);
 			// Recompute adjacency information. Since we change the vertices when we reduce
 			MeshDescriptionHelper.FindOverlappingCorners(MeshDescriptions[LodIndex], OverlappingThreshold);
 			
 			//Make sure the static mesh SectionInfoMap is up to date with the new reduce LOD
 			//We have to remap the material index with the ReductionSettings.BaseLODModel sectionInfoMap
-			
 			//Set the new SectionInfoMap for this reduced LOD base on the ReductionSettings.BaseLODModel SectionInfoMap
-			const FMeshSectionInfoMap& LODModelSectionInfoMap = StaticMesh->SectionInfoMap;
-			const FMeshSectionInfoMap& LODModelOriginalSectionInfoMap = StaticMesh->OriginalSectionInfoMap;
+			TArray<int32> BaseUniqueMaterialIndexes;
+			//Find all unique Material in used order
+			for (const FPolygonGroupID& PolygonGroupID : MeshDescriptions[BaseReduceLodIndex].PolygonGroups().GetElementIDs())
+			{
+				int32 MaterialIndex = StaticMesh->GetMaterialIndexFromImportedMaterialSlotName(BasePolygonGroupImportedMaterialSlotNames[PolygonGroupID]);
+				if (MaterialIndex == INDEX_NONE)
+				{
+					MaterialIndex = PolygonGroupID.GetValue();
+				}
+				BaseUniqueMaterialIndexes.AddUnique(MaterialIndex);
+			}
 			TArray<int32> UniqueMaterialIndex;
 			//Find all unique Material in used order
 			for (const FPolygonGroupID& PolygonGroupID : MeshDescriptions[LodIndex].PolygonGroups().GetElementIDs())
@@ -189,11 +201,12 @@ bool FStaticMeshBuilder::Build(FStaticMeshRenderData& StaticMeshRenderData, USta
 			//If the reduce did not output the same number of section use the base LOD sectionInfoMap
 			bool bIsOldMappingInvalid = OldSectionInfoMapCount != UniqueMaterialIndex.Num();
 
+			bool bValidBaseSectionInfoMap = BeforeBuildSectionInfoMap.GetSectionNumber(BaseReduceLodIndex) > 0;
 			//All used material represent a different section
 			for (int32 SectionIndex = 0; SectionIndex < UniqueMaterialIndex.Num(); ++SectionIndex)
 			{
 				//Keep the old data
-				bool bHasValidLODInfoMap = !bIsOldMappingInvalid && LODModelSectionInfoMap.IsValidSection(LodIndex, SectionIndex);
+				bool bHasValidLODInfoMap = !bIsOldMappingInvalid && BeforeBuildSectionInfoMap.IsValidSection(LodIndex, SectionIndex);
 				//Section material index have to be remap with the ReductionSettings.BaseLODModel SectionInfoMap to create
 				//a valid new section info map for the reduced LOD.
 
@@ -201,19 +214,23 @@ bool FStaticMeshBuilder::Build(FStaticMeshRenderData& StaticMeshRenderData, USta
 				if (!bHasValidLODInfoMap)
 				{
 					bool bSectionInfoSet = false;
-					for (int32 BaseLodSectionIndex = 0; BaseLodSectionIndex < LODModelSectionInfoMap.GetSectionNumber(BaseReduceLodIndex); ++BaseLodSectionIndex)
+					if (bValidBaseSectionInfoMap)
 					{
-						FMeshSectionInfo SectionInfo = LODModelSectionInfoMap.Get(ReductionSettings.BaseLODModel, BaseLodSectionIndex);
-						if (BaseLodSectionIndex == UniqueMaterialIndex[SectionIndex])
+						for (int32 BaseSectionIndex = 0; BaseSectionIndex < BaseUniqueMaterialIndexes.Num(); ++BaseSectionIndex)
 						{
-							//Copy the BaseLODModel section info to the reduce LODIndex.
-							FMeshSectionInfo OriginalSectionInfo = LODModelOriginalSectionInfoMap.Get(ReductionSettings.BaseLODModel, BaseLodSectionIndex);
-							StaticMesh->SectionInfoMap.Set(LodIndex, SectionIndex, SectionInfo);
-							StaticMesh->OriginalSectionInfoMap.Set(LodIndex, SectionIndex, OriginalSectionInfo);
-							bSectionInfoSet = true;
-							break;
+							if (UniqueMaterialIndex[SectionIndex] == BaseUniqueMaterialIndexes[BaseSectionIndex])
+							{
+								//Copy the base sectionInfoMap
+								FMeshSectionInfo SectionInfo = BeforeBuildSectionInfoMap.Get(BaseReduceLodIndex, BaseSectionIndex);
+								FMeshSectionInfo OriginalSectionInfo = BeforeBuildOriginalSectionInfoMap.Get(BaseReduceLodIndex, BaseSectionIndex);
+								StaticMesh->SectionInfoMap.Set(LodIndex, SectionIndex, SectionInfo);
+								StaticMesh->OriginalSectionInfoMap.Set(LodIndex, BaseSectionIndex, OriginalSectionInfo);
+								bSectionInfoSet = true;
+								break;
+							}
 						}
 					}
+
 					if (!bSectionInfoSet)
 					{
 						//Just set the default section info in case we did not found any match with the Base Lod
