Fix reduce SectionInfoMap not set like the reloaded section info map

diff --git a/Engine/Source/Developer/MeshBuilder/Private/StaticMeshBuilder.cpp b/Engine/Source/Developer/MeshBuilder/Private/StaticMeshBuilder.cpp
--- a/Engine/Source/Developer/MeshBuilder/Private/StaticMeshBuilder.cpp
+++ b/Engine/Source/Developer/MeshBuilder/Private/StaticMeshBuilder.cpp
@@ -187,7 +187,7 @@ bool FStaticMeshBuilder::Build(FStaticMeshRenderData& StaticMeshRenderData, USta
 			}
 
 			//If the reduce did not output the same number of section use the base LOD sectionInfoMap
-			bool bIsOldMappingInvalid = OldSectionInfoMapCount != UniqueMaterialIndex.Num();
+			bool bIsOldMappingInvalid = OldSectionInfoMapCount != MeshDescriptions[LodIndex].PolygonGroups().Num();
 
 			//All used material represent a different section
 			for (int32 SectionIndex = 0; SectionIndex < UniqueMaterialIndex.Num(); ++SectionIndex)
