Make MeshBuilder explicitly depend on libstdc++, its libs already do.

diff --git a/Engine/Source/Developer/MeshBuilder/MeshBuilder.Build.cs b/Engine/Source/Developer/MeshBuilder/MeshBuilder.Build.cs
--- a/Engine/Source/Developer/MeshBuilder/MeshBuilder.Build.cs
+++ b/Engine/Source/Developer/MeshBuilder/MeshBuilder.Build.cs
@@ -28,6 +28,8 @@ namespace UnrealBuildTool.Rules
             AddEngineThirdPartyPrivateStaticDependencies(Target, "ForsythTriOptimizer");
 	        AddEngineThirdPartyPrivateStaticDependencies(Target, "nvTessLib");
             AddEngineThirdPartyPrivateStaticDependencies(Target, "QuadricMeshReduction");
-        }
+
+			PublicAdditionalLibraries.Add("stdc++");	// can be fixed, see UE-70769
+       }
 	}
 }
