Fix Win64 build.

diff --git a/Engine/Source/Developer/MeshBuilder/MeshBuilder.Build.cs b/Engine/Source/Developer/MeshBuilder/MeshBuilder.Build.cs
--- a/Engine/Source/Developer/MeshBuilder/MeshBuilder.Build.cs
+++ b/Engine/Source/Developer/MeshBuilder/MeshBuilder.Build.cs
@@ -29,7 +29,10 @@ namespace UnrealBuildTool.Rules
 	        AddEngineThirdPartyPrivateStaticDependencies(Target, "nvTessLib");
             AddEngineThirdPartyPrivateStaticDependencies(Target, "QuadricMeshReduction");
 
-			PublicAdditionalLibraries.Add("stdc++");	// can be fixed, see UE-70769
+			if (Target.IsInPlatformGroup(UnrealPlatformGroup.Unix))
+			{
+				PublicAdditionalLibraries.Add("stdc++");	// can be fixed, see UE-70769
+			}
        }
 	}
 }
