fix static mesh builder wedge index bug

diff --git a/Engine/Source/Developer/MeshBuilder/Private/StaticMeshBuilder.cpp b/Engine/Source/Developer/MeshBuilder/Private/StaticMeshBuilder.cpp
--- a/Engine/Source/Developer/MeshBuilder/Private/StaticMeshBuilder.cpp
+++ b/Engine/Source/Developer/MeshBuilder/Private/StaticMeshBuilder.cpp
@@ -451,6 +451,7 @@ void BuildVertexBuffer(
 				|| PointsEqual(CornerPositions[0], CornerPositions[2], OverlappingThresholds)
 				|| PointsEqual(CornerPositions[1], CornerPositions[2], OverlappingThresholds))
 			{
+				WedgeIndex += 3;
 				continue;
 			}
 
