#UE4 Quick fix for crash in content validation

diff --git a/Engine/Source/Developer/MeshBuilder/Private/SkeletalMeshBuilder.cpp b/Engine/Source/Developer/MeshBuilder/Private/SkeletalMeshBuilder.cpp
--- a/Engine/Source/Developer/MeshBuilder/Private/SkeletalMeshBuilder.cpp
+++ b/Engine/Source/Developer/MeshBuilder/Private/SkeletalMeshBuilder.cpp
@@ -76,7 +76,10 @@ namespace SkeletalMeshBuilderHelperNS
 		//Update all the faces
 		for (int32 FaceIndex = 0; FaceIndex < LODFaces.Num(); ++FaceIndex)
 		{
-			LODFaces[FaceIndex].MeshMaterialIndex = MaterialRemap[LODFaces[FaceIndex].MeshMaterialIndex];
+			if (MaterialRemap.IsValidIndex(LODFaces[FaceIndex].MeshMaterialIndex))
+			{
+				LODFaces[FaceIndex].MeshMaterialIndex = MaterialRemap[LODFaces[FaceIndex].MeshMaterialIndex];
+			}
 		}
 	}
 }
