Fix old asset build issue with faces material remap, in case there is some data missing

diff --git a/Engine/Source/Developer/MeshBuilder/Private/SkeletalMeshBuilder.cpp b/Engine/Source/Developer/MeshBuilder/Private/SkeletalMeshBuilder.cpp
--- a/Engine/Source/Developer/MeshBuilder/Private/SkeletalMeshBuilder.cpp
+++ b/Engine/Source/Developer/MeshBuilder/Private/SkeletalMeshBuilder.cpp
@@ -56,9 +56,17 @@ namespace SkeletalMeshBuilderHelperNS
 {
 	void FixFaceMaterial(USkeletalMesh* SkeletalMesh, TArray<SkeletalMeshImportData::FMaterial>& RawMeshMaterials, TArray<SkeletalMeshImportData::FMeshFace>& LODFaces)
 	{
+		if (RawMeshMaterials.Num() <= 1)
+		{
+			//Nothing to fix if we have 1 or less material
+			return;
+		}
+		
 		//Fix the material for the faces
 		TArray<int32> MaterialRemap;
 		MaterialRemap.Reserve(RawMeshMaterials.Num());
+		//Optimization to avoid doing the remap if no material have to change
+		bool bNeedRemapping = false;
 		for (int32 MaterialIndex = 0; MaterialIndex < RawMeshMaterials.Num(); ++MaterialIndex)
 		{
 			MaterialRemap.Add(MaterialIndex);
@@ -68,15 +76,29 @@ namespace SkeletalMeshBuilderHelperNS
 				FName MeshMaterialName = SkeletalMesh->Materials[MeshMaterialIndex].ImportedMaterialSlotName;
 				if (MaterialImportName == MeshMaterialName)
 				{
+					if (MaterialRemap[MaterialIndex] != MeshMaterialIndex)
+					{
+						bNeedRemapping = true;
+					}
 					MaterialRemap[MaterialIndex] = MeshMaterialIndex;
 					break;
 				}
 			}
 		}
-		//Update all the faces
-		for (int32 FaceIndex = 0; FaceIndex < LODFaces.Num(); ++FaceIndex)
+		if (bNeedRemapping)
 		{
-			if (MaterialRemap.IsValidIndex(LODFaces[FaceIndex].MeshMaterialIndex))
+			//Make sure the data is good before doing the change, We cannot do the remap if we
+			//have a bad synchronization between the face data and the Materials data.
+			for (int32 FaceIndex = 0; FaceIndex < LODFaces.Num(); ++FaceIndex)
+			{
+				if (!MaterialRemap.IsValidIndex(LODFaces[FaceIndex].MeshMaterialIndex))
+				{
+					return;
+				}
+			}
+
+			//Update all the faces
+			for (int32 FaceIndex = 0; FaceIndex < LODFaces.Num(); ++FaceIndex)
 			{
 				LODFaces[FaceIndex].MeshMaterialIndex = MaterialRemap[LODFaces[FaceIndex].MeshMaterialIndex];
 			}
