Fix a crash when importing a skeletal mesh with no UVs
#rnx

diff --git a/Engine/Source/Developer/MeshBuilder/Private/SkeletalMeshBuilder.cpp b/Engine/Source/Developer/MeshBuilder/Private/SkeletalMeshBuilder.cpp
--- a/Engine/Source/Developer/MeshBuilder/Private/SkeletalMeshBuilder.cpp
+++ b/Engine/Source/Developer/MeshBuilder/Private/SkeletalMeshBuilder.cpp
@@ -135,6 +135,7 @@ bool FSkeletalMeshBuilder::Build(USkeletalMesh* SkeletalMesh, const int32 LODInd
 
 	FSkeletalMeshImportData SkeletalMeshImportData;
 	float ProgressStepSize = SkeletalMesh->IsReductionActive(LODIndex) ? 0.5f : 1.0f;
+	int32 NumTextCoord = 1; //We need to send rendering at least one tex coord buffer
 
 	//This scope define where we can use the LODModel, after a reduction the LODModel must be requery since it is a new instance
 	{
@@ -149,6 +150,9 @@ bool FSkeletalMeshBuilder::Build(USkeletalMesh* SkeletalMesh, const int32 LODInd
 		TArray<int32> LODPointToRawMap;
 		SkeletalMeshImportData.CopyLODImportData(LODPoints, LODWedges, LODFaces, LODInfluences, LODPointToRawMap);
 
+		//Use the max because we need to have at least one texture coordinnate
+		NumTextCoord = FMath::Max<int32>(NumTextCoord, SkeletalMeshImportData.NumTexCoords);
+
 		SkeletalMeshBuilderHelperNS::FixFaceMaterial(SkeletalMesh, SkeletalMeshImportData.Materials, LODFaces);
 
 		//Build the skeletalmesh using mesh utilities module
@@ -173,7 +177,7 @@ bool FSkeletalMeshBuilder::Build(USkeletalMesh* SkeletalMesh, const int32 LODInd
 		BuildLODModel.SyncronizeUserSectionsDataArray();
 
 		// Set texture coordinate count on the new model.
-		BuildLODModel.NumTexCoords = SkeletalMeshImportData.NumTexCoords;
+		BuildLODModel.NumTexCoords = NumTextCoord;
 
 		//Re-apply the morph target
 		SlowTask.EnterProgressFrame(1.0f, NSLOCTEXT("SkeltalMeshBuilder", "RebuildMorphTarget", "Rebuilding morph targets..."));
@@ -251,7 +255,7 @@ bool FSkeletalMeshBuilder::Build(USkeletalMesh* SkeletalMesh, const int32 LODInd
 	FLODUtilities::RestoreClothingFromBackup(SkeletalMesh, ClothingBindings, LODIndex);
 
 	LODModelAfterReduction.SyncronizeUserSectionsDataArray();
-	LODModelAfterReduction.NumTexCoords = SkeletalMeshImportData.NumTexCoords;
+	LODModelAfterReduction.NumTexCoords = NumTextCoord;
 	LODModelAfterReduction.BuildStringID = BackupBuildStringID;
 
 	SlowTask.EnterProgressFrame(ProgressStepSize, NSLOCTEXT("SkeltalMeshBuilder", "RegenerateDependentLODs", "Regenerate Dependent LODs..."));
