Fix ensure in CIS anti cheat build
#rnx

diff --git a/Engine/Source/Developer/MeshBuilder/Private/SkeletalMeshBuilder.cpp b/Engine/Source/Developer/MeshBuilder/Private/SkeletalMeshBuilder.cpp
--- a/Engine/Source/Developer/MeshBuilder/Private/SkeletalMeshBuilder.cpp
+++ b/Engine/Source/Developer/MeshBuilder/Private/SkeletalMeshBuilder.cpp
@@ -124,7 +124,7 @@ bool FSkeletalMeshBuilder::Build(USkeletalMesh* SkeletalMesh, const int32 LODInd
 
 	const FReferenceSkeleton& RefSkeleton = SkeletalMesh->RefSkeleton;
 
-	FScopedSlowTask SlowTask(6.0f, NSLOCTEXT("SkeltalMeshBuilder", "BuildingSkeletalMeshLOD", "Building skeletal mesh LOD"));
+	FScopedSlowTask SlowTask(6.01f, NSLOCTEXT("SkeltalMeshBuilder", "BuildingSkeletalMeshLOD", "Building skeletal mesh LOD"));
 	SlowTask.MakeDialog();
 
 	//Prevent any PostEdit change during the build
@@ -134,7 +134,6 @@ bool FSkeletalMeshBuilder::Build(USkeletalMesh* SkeletalMesh, const int32 LODInd
 	FLODUtilities::UnbindClothingAndBackup(SkeletalMesh, ClothingBindings, LODIndex);
 
 	FSkeletalMeshImportData SkeletalMeshImportData;
-	float ProgressStepSize = SkeletalMesh->IsReductionActive(LODIndex) ? 0.5f : 1.0f;
 	int32 NumTextCoord = 1; //We need to send rendering at least one tex coord buffer
 
 	//This scope define where we can use the LODModel, after a reduction the LODModel must be requery since it is a new instance
@@ -202,7 +201,7 @@ bool FSkeletalMeshBuilder::Build(USkeletalMesh* SkeletalMesh, const int32 LODInd
 		//We are reduce ourself in this case we reduce ourself from the original data and return true.
 		if (SkeletalMesh->IsReductionActive(LODIndex))
 		{
-			SlowTask.EnterProgressFrame(ProgressStepSize, NSLOCTEXT("SkeltalMeshBuilder", "RegenerateLOD", "Regenerate LOD..."));
+			SlowTask.EnterProgressFrame(1.0f, NSLOCTEXT("SkeltalMeshBuilder", "RegenerateLOD", "Regenerate LOD..."));
 			//Update the original reduction data since we just build a new LODModel.
 			if (LODInfo->ReductionSettings.BaseLOD == LODIndex && SkeletalMesh->GetImportedModel()->OriginalReductionSourceMeshData.IsValidIndex(LODIndex))
 			{
@@ -259,14 +258,11 @@ bool FSkeletalMeshBuilder::Build(USkeletalMesh* SkeletalMesh, const int32 LODInd
 	LODModelAfterReduction.NumTexCoords = NumTextCoord;
 	LODModelAfterReduction.BuildStringID = BackupBuildStringID;
 
-	SlowTask.EnterProgressFrame(ProgressStepSize, NSLOCTEXT("SkeltalMeshBuilder", "RegenerateDependentLODs", "Regenerate Dependent LODs..."));
+	SlowTask.EnterProgressFrame(1.0f, NSLOCTEXT("SkeltalMeshBuilder", "RegenerateDependentLODs", "Regenerate Dependent LODs..."));
 	if (bRegenDepLODs)
 	{
 		//Regenerate dependent LODs
 		FLODUtilities::RegenerateDependentLODs(SkeletalMesh, LODIndex);
 	}
-
-	SlowTask.EnterProgressFrame();
-
 	return true;
 }
