when building skeletalmesh, force the tangent and normal options to true if the source data do not have the tangents or normals data.

diff --git a/Engine/Source/Developer/MeshBuilder/Private/SkeletalMeshBuilder.cpp b/Engine/Source/Developer/MeshBuilder/Private/SkeletalMeshBuilder.cpp
--- a/Engine/Source/Developer/MeshBuilder/Private/SkeletalMeshBuilder.cpp
+++ b/Engine/Source/Developer/MeshBuilder/Private/SkeletalMeshBuilder.cpp
@@ -157,6 +157,9 @@ bool FSkeletalMeshBuilder::Build(USkeletalMesh* SkeletalMesh, const int32 LODInd
 		//Build the skeletalmesh using mesh utilities module
 		IMeshUtilities::MeshBuildOptions Options;
 		Options.FillOptions(LODInfo->BuildSettings);
+		//Force the normals or tangent in case the data is missing
+		Options.bComputeNormals |= !SkeletalMeshImportData.bHasNormals;
+		Options.bComputeTangents |= !SkeletalMeshImportData.bHasTangents;
 		IMeshUtilities& MeshUtilities = FModuleManager::Get().LoadModuleChecked<IMeshUtilities>("MeshUtilities");
 
 		// Create skinning streams for NewModel.
