Memory optims for mesh building

diff --git a/Engine/Source/Developer/MeshBuilder/Private/StaticMeshBuilder.cpp b/Engine/Source/Developer/MeshBuilder/Private/StaticMeshBuilder.cpp
--- a/Engine/Source/Developer/MeshBuilder/Private/StaticMeshBuilder.cpp
+++ b/Engine/Source/Developer/MeshBuilder/Private/StaticMeshBuilder.cpp
@@ -280,6 +280,13 @@ bool FStaticMeshBuilder::Build(FStaticMeshRenderData& StaticMeshRenderData, USta
 
 		const uint32 NumTextureCoord = MeshDescriptions[LodIndex].VertexInstanceAttributes().GetAttributesRef<FVector2D>( MeshAttribute::VertexInstance::TextureCoordinate ).GetNumChannels();
 
+		// Only the render data and vertex buffers will be used from now on unless we have more than one source models
+		// This will help with memory usage for Nanite Mesh by releasing memory before doing the build
+		if (NumSourceModels == 1)
+		{
+			MeshDescriptions.Empty();
+		}
+
 		// Concatenate the per-section index buffers.
 		TArray<uint32> CombinedIndices;
 		bool bNeeds32BitIndices = false;
