Replaced descriptor get with a clone inside FStaticMeshBuilder::BuildMeshVertexPositions to lower memory usage during the mesh build process.

diff --git a/Engine/Source/Developer/MeshBuilder/Private/StaticMeshBuilder.cpp b/Engine/Source/Developer/MeshBuilder/Private/StaticMeshBuilder.cpp
--- a/Engine/Source/Developer/MeshBuilder/Private/StaticMeshBuilder.cpp
+++ b/Engine/Source/Developer/MeshBuilder/Private/StaticMeshBuilder.cpp
@@ -471,44 +471,52 @@ bool FStaticMeshBuilder::BuildMeshVertexPositions(
 {
 	TRACE_CPUPROFILER_EVENT_SCOPE_TEXT(TEXT("FStaticMeshBuilder::BuildMeshVertexPositions"));
 
-	const FMeshDescription* MeshDescription = StaticMesh->GetMeshDescription(0);
-	if (MeshDescription == nullptr)
+	if (!StaticMesh->IsMeshDescriptionValid(0))
 	{
 		//Warn the user that there is no mesh description data
 		UE_LOG(LogStaticMeshBuilder, Error, TEXT("Cannot find a valid mesh description to build the asset."));
 		return false;
 	}
 
-	const FMeshBuildSettings& BuildSettings = StaticMesh->GetSourceModel(0).BuildSettings;
+	const int32 NumSourceModels = StaticMesh->GetNumSourceModels();
+	if (NumSourceModels > 0)
+	{
+		FMeshDescription MeshDescription;
+		const bool bIsMeshDescriptionValid = StaticMesh->CloneMeshDescription(/*LodIndex*/ 0, MeshDescription);
+		if (bIsMeshDescriptionValid)
+		{
+			const FMeshBuildSettings& BuildSettings = StaticMesh->GetSourceModel(0).BuildSettings;
 
-	const FStaticMeshConstAttributes Attributes(*MeshDescription);
-	TArrayView<const FVector> VertexPositions = Attributes.GetVertexPositions().GetRawArray();
-	TArrayView<const FVertexID> VertexIndices = Attributes.GetTriangleVertexIndices().GetRawArray();
+			const FStaticMeshConstAttributes Attributes(MeshDescription);
+			TArrayView<const FVector> VertexPositions = Attributes.GetVertexPositions().GetRawArray();
+			TArrayView<const FVertexID> VertexIndices = Attributes.GetTriangleVertexIndices().GetRawArray();
 
-	BuiltVertices.Reserve(VertexPositions.Num());
-	for (int32 VertexIndex = 0; VertexIndex < VertexPositions.Num(); ++VertexIndex)
-	{
-		BuiltVertices.Add(VertexPositions[VertexIndex] * BuildSettings.BuildScale3D);
-	}
+			BuiltVertices.Reserve(VertexPositions.Num());
+			for (int32 VertexIndex = 0; VertexIndex < VertexPositions.Num(); ++VertexIndex)
+			{
+				BuiltVertices.Add(VertexPositions[VertexIndex] * BuildSettings.BuildScale3D);
+			}
 
-	BuiltIndices.Reserve(VertexIndices.Num());
-	for (int32 TriangleIndex = 0; TriangleIndex < VertexIndices.Num() / 3; ++TriangleIndex)
-	{
-		const uint32 I0 = VertexIndices[TriangleIndex * 3 + 0];
-		const uint32 I1 = VertexIndices[TriangleIndex * 3 + 1];
-		const uint32 I2 = VertexIndices[TriangleIndex * 3 + 2];
+			BuiltIndices.Reserve(VertexIndices.Num());
+			for (int32 TriangleIndex = 0; TriangleIndex < VertexIndices.Num() / 3; ++TriangleIndex)
+			{
+				const uint32 I0 = VertexIndices[TriangleIndex * 3 + 0];
+				const uint32 I1 = VertexIndices[TriangleIndex * 3 + 1];
+				const uint32 I2 = VertexIndices[TriangleIndex * 3 + 2];
 
-		const FVector V0 = BuiltVertices[I0];
-		const FVector V1 = BuiltVertices[I1];
-		const FVector V2 = BuiltVertices[I2];
+				const FVector V0 = BuiltVertices[I0];
+				const FVector V1 = BuiltVertices[I1];
+				const FVector V2 = BuiltVertices[I2];
 
-		const FVector TriangleNormal = ((V1 - V2) ^ (V0 - V2));
-		const bool bDegenerateTriangle = TriangleNormal.SizeSquared() < SMALL_NUMBER;
-		if (!bDegenerateTriangle)
-		{
-			BuiltIndices.Add(I0);
-			BuiltIndices.Add(I1);
-			BuiltIndices.Add(I2);
+				const FVector TriangleNormal = ((V1 - V2) ^ (V0 - V2));
+				const bool bDegenerateTriangle = TriangleNormal.SizeSquared() < SMALL_NUMBER;
+				if (!bDegenerateTriangle)
+				{
+					BuiltIndices.Add(I0);
+					BuiltIndices.Add(I1);
+					BuiltIndices.Add(I2);
+				}
+			}
 		}
 	}
 
