Removed FMeshBuilder::Build(UStaticMesh* StaticMesh, TArray<FStaticMeshBuildVertex>& Verts, TArray<uint32>& Indexes, TArray<FStaticMeshSection>& Sections, uint32& NumTexCoords, bool& bHasColors). It's somewhat danegous as it leaks mesh descriptors and isn't used anymore.

diff --git a/Engine/Source/Developer/MeshBuilder/Private/MeshBuilderModule.cpp b/Engine/Source/Developer/MeshBuilder/Private/MeshBuilderModule.cpp
--- a/Engine/Source/Developer/MeshBuilder/Private/MeshBuilderModule.cpp
+++ b/Engine/Source/Developer/MeshBuilder/Private/MeshBuilderModule.cpp
@@ -28,14 +28,6 @@ public:
 
 	virtual bool BuildMesh(FStaticMeshRenderData& OutRenderData, UObject* Mesh, const FStaticMeshLODGroup& LODGroup) override;
 
-	virtual bool BuildMesh(
-		UObject* StaticMesh,
-		TArray< FStaticMeshBuildVertex >& Verts,
-		TArray< uint32 >& Indexes,
-		FStaticMeshSectionArray& Sections,
-		uint32& NumTexCoords,
-		bool& bHasColors ) override;
-
 	virtual bool BuildMeshVertexPositions(
 		UObject* StaticMesh,
 		TArray<uint32>& Indices,
@@ -60,23 +52,6 @@ bool FMeshBuilderModule::BuildMesh(FStaticMeshRenderData& OutRenderData, class U
 	return false;
 }
 
-bool FMeshBuilderModule::BuildMesh(
-	UObject* Mesh,
-	TArray< FStaticMeshBuildVertex >& Verts,
-	TArray< uint32 >& Indexes,
-	FStaticMeshSectionArray& Sections,
-	uint32& NumTexCoords,
-	bool& bHasColors)
-{
-	UStaticMesh* StaticMesh = Cast< UStaticMesh >(Mesh);
-	if (StaticMesh)
-	{
-		//Call the static mesh builder
-		return FStaticMeshBuilder().Build(StaticMesh, Verts, Indexes, Sections, NumTexCoords, bHasColors);
-	}
-	return false;
-}
-
 bool FMeshBuilderModule::BuildMeshVertexPositions(
 	UObject* Mesh,
 	TArray<uint32>& Indices,
diff --git a/Engine/Source/Developer/MeshBuilder/Private/StaticMeshBuilder.cpp b/Engine/Source/Developer/MeshBuilder/Private/StaticMeshBuilder.cpp
--- a/Engine/Source/Developer/MeshBuilder/Private/StaticMeshBuilder.cpp
+++ b/Engine/Source/Developer/MeshBuilder/Private/StaticMeshBuilder.cpp
@@ -373,97 +373,6 @@ bool FStaticMeshBuilder::Build(FStaticMeshRenderData& StaticMeshRenderData, USta
 	return true;
 }
 
-bool FStaticMeshBuilder::Build(
-	UStaticMesh* StaticMesh,
-	TArray< FStaticMeshBuildVertex >& Verts,
-	TArray< uint32 >& Indexes,
-	FStaticMeshSectionArray& Sections,
-	uint32& NumTexCoords,
-	bool& bHasColors )
-{
-	TRACE_CPUPROFILER_EVENT_SCOPE_TEXT(TEXT("FStaticMeshBuilder::Build"));
-
-	if (StaticMesh->GetMeshDescription(0) == nullptr)
-	{
-		//Warn the user that there is no mesh description data
-		UE_LOG(LogStaticMeshBuilder, Error, TEXT("Cannot find a valid mesh description to build the asset."));
-		return false;
-	}
-
-	FStaticMeshSourceModel& SourceModel = StaticMesh->GetSourceModel(0);
-	FMeshBuildSettings& BuildSettings = SourceModel.BuildSettings;
-	
-	FMeshDescription		MeshDescription;
-	FMeshDescriptionHelper	MeshDescriptionHelper( &BuildSettings );
-
-	const FMeshDescription* OriginalMeshDescription = StaticMesh->GetMeshDescription(0);
-	if (OriginalMeshDescription != nullptr)
-	{
-		TRACE_CPUPROFILER_EVENT_SCOPE_TEXT(TEXT("FStaticMeshBuilder::Build::Mappings"));
-		MeshDescription = *OriginalMeshDescription;
-		MeshDescriptionHelper.SetupRenderMeshDescription(StaticMesh, MeshDescription);
-	}
-	else
-	{
-		TRACE_CPUPROFILER_EVENT_SCOPE_TEXT(TEXT("FStaticMeshBuilder::Build::Mappings"));
-		//Set the overlapping threshold
-		float ComparisonThreshold = BuildSettings.bRemoveDegenerates ? THRESH_POINTS_ARE_SAME : 0.0f;
-		MeshDescriptionHelper.FindOverlappingCorners( MeshDescription, ComparisonThreshold );
-	}
-
-	const FPolygonGroupArray& PolygonGroups = MeshDescription.PolygonGroups();
-
-	Sections.Empty( PolygonGroups.Num() );
-	
-	TArray<int32> RemapVerts;
-	TArray<int32> WedgeMap;
-
-	//Prepare the PerSectionIndices array so we can optimize the index buffer for the GPU
-	TArray< TArray<uint32> > PerSectionIndices;
-	PerSectionIndices.AddDefaulted( PolygonGroups.Num() );
-
-	//Build the vertex and index buffer
-	BuildVertexBuffer( StaticMesh, MeshDescription, BuildSettings, WedgeMap, Sections, PerSectionIndices, Verts, MeshDescriptionHelper.GetOverlappingCorners(), RemapVerts );
-
-	NumTexCoords = MeshDescription.VertexInstanceAttributes().GetAttributesRef< FVector2D >( MeshAttribute::VertexInstance::TextureCoordinate ).GetNumChannels();
-
-	// Concatenate the per-section index buffers.
-	for (int32 SectionIndex = 0; SectionIndex < Sections.Num(); SectionIndex++)
-	{
-		TRACE_CPUPROFILER_EVENT_SCOPE_TEXT(TEXT("FStaticMeshBuilder::Build::CombineIndices"));
-
-		FStaticMeshSection& Section = Sections[ SectionIndex ];
-		TArray<uint32> const& SectionIndices = PerSectionIndices[ SectionIndex ];
-		Section.FirstIndex = 0;
-		Section.NumTriangles = 0;
-		Section.MinVertexIndex = 0;
-		Section.MaxVertexIndex = 0;
-
-		if (SectionIndices.Num())
-		{
-			Section.FirstIndex = Indexes.Num();
-			Section.NumTriangles = SectionIndices.Num() / 3;
-
-			Indexes.AddUninitialized( SectionIndices.Num() );
-			uint32* DestPtr = &Indexes[ Section.FirstIndex ];
-			uint32 const* SrcPtr = SectionIndices.GetData();
-
-			Section.MinVertexIndex = *SrcPtr;
-			Section.MaxVertexIndex = *SrcPtr;
-
-			for (int32 Index = 0; Index < SectionIndices.Num(); Index++)
-			{
-				uint32 VertIndex = *SrcPtr++;
-				Section.MinVertexIndex = FMath::Min<uint32>(VertIndex, Section.MinVertexIndex);
-				Section.MaxVertexIndex = FMath::Max<uint32>(VertIndex, Section.MaxVertexIndex);
-				*DestPtr++ = VertIndex;
-			}
-		}
-	}
-
-	return true;
-}
-
 bool FStaticMeshBuilder::BuildMeshVertexPositions(
 	UStaticMesh* StaticMesh,
 	TArray<uint32>& BuiltIndices,
diff --git a/Engine/Source/Developer/MeshBuilder/Public/IMeshBuilderModule.h b/Engine/Source/Developer/MeshBuilder/Public/IMeshBuilderModule.h
--- a/Engine/Source/Developer/MeshBuilder/Public/IMeshBuilderModule.h
+++ b/Engine/Source/Developer/MeshBuilder/Public/IMeshBuilderModule.h
@@ -31,17 +31,6 @@ public:
 
 	virtual bool BuildMesh(class FStaticMeshRenderData& OutRenderData, class UObject* Mesh, const class FStaticMeshLODGroup& LODGroup) { return false; }
 
-	virtual bool BuildMesh(
-		class UObject* StaticMesh,
-		TArray< FStaticMeshBuildVertex >& Verts,
-		TArray< uint32 >& Indexes,
-		FStaticMeshSectionArray& Sections,
-		uint32& NumTexCoords,
-		bool& bHasColors )
-	{
-		return false;
-	}
-
 	virtual bool BuildMeshVertexPositions(
 		class UObject* StaticMesh,
 		TArray<uint32>& Indices,
diff --git a/Engine/Source/Developer/MeshBuilder/Public/MeshBuilder.h b/Engine/Source/Developer/MeshBuilder/Public/MeshBuilder.h
--- a/Engine/Source/Developer/MeshBuilder/Public/MeshBuilder.h
+++ b/Engine/Source/Developer/MeshBuilder/Public/MeshBuilder.h
@@ -26,13 +26,6 @@ public:
 	 * Build function should be override and is the starting point for static mesh builders
 	 */
 	virtual bool Build(FStaticMeshRenderData& OutRenderData, UStaticMesh* StaticMesh, const FStaticMeshLODGroup& LODGroup) = 0;
-	virtual bool Build(
-		UStaticMesh* StaticMesh,
-		TArray< FStaticMeshBuildVertex >& Verts,
-		TArray< uint32 >& Indexes,
-		FStaticMeshSectionArray& Sections,
-		uint32& NumTexCoords,
-		bool& bHasColors ) = 0;
 	virtual bool BuildMeshVertexPositions(
 		UStaticMesh* StaticMesh,
 		TArray<uint32>& Indices,
diff --git a/Engine/Source/Developer/MeshBuilder/Public/SkeletalMeshBuilder.h b/Engine/Source/Developer/MeshBuilder/Public/SkeletalMeshBuilder.h
--- a/Engine/Source/Developer/MeshBuilder/Public/SkeletalMeshBuilder.h
+++ b/Engine/Source/Developer/MeshBuilder/Public/SkeletalMeshBuilder.h
@@ -25,19 +25,6 @@ public:
 		return false;
 	}
 
-	virtual bool Build(
-		UStaticMesh* StaticMesh,
-		TArray< FStaticMeshBuildVertex >& Verts,
-		TArray< uint32 >& Indexes,
-		FStaticMeshSectionArray& Sections,
-		uint32& NumTexCoords,
-		bool& bHasColors) override
-	{
-		bool No_Support_For_StaticMesh_Build_In_FSkeletalMeshBuilder_Class = false;
-		check(No_Support_For_StaticMesh_Build_In_FSkeletalMeshBuilder_Class);
-		return false;
-	}
-
 	virtual bool BuildMeshVertexPositions(
 		UStaticMesh* StaticMesh,
 		TArray<uint32>& Indices,
diff --git a/Engine/Source/Developer/MeshBuilder/Public/StaticMeshBuilder.h b/Engine/Source/Developer/MeshBuilder/Public/StaticMeshBuilder.h
--- a/Engine/Source/Developer/MeshBuilder/Public/StaticMeshBuilder.h
+++ b/Engine/Source/Developer/MeshBuilder/Public/StaticMeshBuilder.h
@@ -28,14 +28,6 @@ public:
 		return false;
 	}
 
-	virtual bool Build(
-		UStaticMesh* StaticMesh,
-		TArray< FStaticMeshBuildVertex >& Verts,
-		TArray< uint32 >& Indexes,
-		FStaticMeshSectionArray& Sections,
-		uint32& NumTexCoords,
-		bool& bHasColors ) override;
-
 	virtual bool BuildMeshVertexPositions(
 		UStaticMesh* StaticMesh,
 		TArray<uint32>& Indices,
