Add Insights tag in Skeleton Simplifier

diff --git a/Engine/Plugins/Experimental/SkeletalReduction/Source/Private/SkeletalSimplifier.cpp b/Engine/Plugins/Experimental/SkeletalReduction/Source/Private/SkeletalSimplifier.cpp
--- a/Engine/Plugins/Experimental/SkeletalReduction/Source/Private/SkeletalSimplifier.cpp
+++ b/Engine/Plugins/Experimental/SkeletalReduction/Source/Private/SkeletalSimplifier.cpp
@@ -106,6 +106,7 @@ void SkeletalSimplifier::FMeshSimplifier::SetBoxCornersLocked()
 
 void SkeletalSimplifier::FMeshSimplifier::InitCosts()
 {
+	TRACE_CPUPROFILER_EVENT_SCOPE(SkeletalSimplifier::InitCosts);
 
 	const int32 NumEdges = MeshManager.TotalNumEdges();
 	for (int i = 0; i < NumEdges; ++i)
diff --git a/Engine/Plugins/Experimental/SkeletalReduction/Source/Private/SkeletalSimplifier.h b/Engine/Plugins/Experimental/SkeletalReduction/Source/Private/SkeletalSimplifier.h
--- a/Engine/Plugins/Experimental/SkeletalReduction/Source/Private/SkeletalSimplifier.h
+++ b/Engine/Plugins/Experimental/SkeletalReduction/Source/Private/SkeletalSimplifier.h
@@ -442,6 +442,7 @@ namespace SkeletalSimplifier
 	template< typename TerminationCriterionType>
 	float FMeshSimplifier::SimplifyMesh(TerminationCriterionType TerminationCriterion)
 	{
+		TRACE_CPUPROFILER_EVENT_SCOPE(FMeshSimplifier::SimplifyMesh);
 
 		// Build the cost heap
 		InitCosts();
diff --git a/Engine/Source/Developer/MeshBuilder/Private/SkeletalMeshBuilder.cpp b/Engine/Source/Developer/MeshBuilder/Private/SkeletalMeshBuilder.cpp
--- a/Engine/Source/Developer/MeshBuilder/Private/SkeletalMeshBuilder.cpp
+++ b/Engine/Source/Developer/MeshBuilder/Private/SkeletalMeshBuilder.cpp
@@ -57,6 +57,8 @@ FSkeletalMeshBuilder::FSkeletalMeshBuilder()
 
 bool FSkeletalMeshBuilder::Build(USkeletalMesh* SkeletalMesh, const int32 LODIndex, const bool bRegenDepLODs)
 {
+	TRACE_CPUPROFILER_EVENT_SCOPE(FSkeletalMeshBuilder::Build);
+
 	check(SkeletalMesh->GetImportedModel());
 	check(SkeletalMesh->GetImportedModel()->LODModels.IsValidIndex(LODIndex));
 	check(SkeletalMesh->GetLODInfo(LODIndex) != nullptr);
