Support using HiRes SourceModel to generate Nanite rendering data in Static Mesh Build. If HiRes exists, it will be used instead of LOD0, in local function BuildNaniteFromHiResSourceModel().  Refactored several blocks out of FStaticMeshBuilder::Build() into standalone functions so they could be called in multiple places.

This CL also updates the bounding box computation, previously the bounds were computed on LOD0 render data, ie after the nanite cut. Now the bounds of either the HiRes SourceModel, or high-res LOD0, are combined with the final LOD0 bounds.

#rnx

diff --git a/Engine/Source/Developer/MeshBuilder/Private/StaticMeshBuilder.cpp b/Engine/Source/Developer/MeshBuilder/Private/StaticMeshBuilder.cpp
--- a/Engine/Source/Developer/MeshBuilder/Private/StaticMeshBuilder.cpp
+++ b/Engine/Source/Developer/MeshBuilder/Private/StaticMeshBuilder.cpp
@@ -53,9 +53,190 @@ static bool UseNativeQuadraticReduction()
 	return bUseQuadricSimplier;
 }
 
+
+/**
+ * Compute bounding box and sphere from position buffer
+ */
+static void ComputeBoundsFromPositionBuffer(const FPositionVertexBuffer& UsePositionBuffer, FVector& OriginOut, FVector& ExtentOut, float& RadiusOut)
+{
+	// Calculate the bounding box.
+	FBox BoundingBox(ForceInit);
+	for (uint32 VertexIndex = 0; VertexIndex < UsePositionBuffer.GetNumVertices(); VertexIndex++)
+	{
+		BoundingBox += UsePositionBuffer.VertexPosition(VertexIndex);
+	}
+	BoundingBox.GetCenterAndExtents(OriginOut, ExtentOut);
+
+	// Calculate the bounding sphere, using the center of the bounding box as the origin.
+	RadiusOut = 0.0f;
+	for (uint32 VertexIndex = 0; VertexIndex < UsePositionBuffer.GetNumVertices(); VertexIndex++)
+	{
+		RadiusOut = FMath::Max(
+			(UsePositionBuffer.VertexPosition(VertexIndex) - OriginOut).Size(),
+			RadiusOut
+		);
+	}
+}
+
+
+/**
+ * Compute bounding box and sphere from vertices
+ */
+static void ComputeBoundsFromVertexList(const TArray<FStaticMeshBuildVertex>& Vertices, FVector& OriginOut, FVector& ExtentOut, float& RadiusOut)
+{
+	// Calculate the bounding box.
+	FBox BoundingBox(ForceInit);
+	for (int32 VertexIndex = 0; VertexIndex < Vertices.Num(); VertexIndex++)
+	{
+		BoundingBox += Vertices[VertexIndex].Position;
+	}
+	BoundingBox.GetCenterAndExtents(OriginOut, ExtentOut);
+
+	// Calculate the bounding sphere, using the center of the bounding box as the origin.
+	RadiusOut = 0.0f;
+	for (int32 VertexIndex = 0; VertexIndex < Vertices.Num(); VertexIndex++)
+	{
+		RadiusOut = FMath::Max(
+			(Vertices[VertexIndex].Position - OriginOut).Size(),
+			RadiusOut
+		);
+	}
+}
+
+
+
+/**
+ * Utility function used inside FStaticMeshBuilder::Build() per-LOD loop to populate
+ * the Sections in a FStaticMeshLODResources from PerSectionIndices, as well as
+ * concatenate all section indices into CombinedIndicesOut.
+ * Returned bNeeds32BitIndicesOut indicates whether max vert index is larger than max int16
+ */
+static void BuildCombinedSectionIndices(
+	const TArray<TArray<uint32>>& PerSectionIndices, 
+	FStaticMeshLODResources& StaticMeshLODInOut, 
+	TArray<uint32>& CombinedIndicesOut,
+	bool& bNeeds32BitIndicesOut )
+{
+	bNeeds32BitIndicesOut = false;
+	for (int32 SectionIndex = 0; SectionIndex < StaticMeshLODInOut.Sections.Num(); SectionIndex++)
+	{
+		FStaticMeshSection& Section = StaticMeshLODInOut.Sections[SectionIndex];
+		const TArray<uint32>& SectionIndices = PerSectionIndices[SectionIndex];
+		Section.FirstIndex = 0;
+		Section.NumTriangles = 0;
+		Section.MinVertexIndex = 0;
+		Section.MaxVertexIndex = 0;
+
+		if (SectionIndices.Num())
+		{
+			Section.FirstIndex = CombinedIndicesOut.Num();
+			Section.NumTriangles = SectionIndices.Num() / 3;
+
+			CombinedIndicesOut.AddUninitialized(SectionIndices.Num());
+			uint32* DestPtr = &CombinedIndicesOut[Section.FirstIndex];
+			uint32 const* SrcPtr = SectionIndices.GetData();
+
+			Section.MinVertexIndex = *SrcPtr;
+			Section.MaxVertexIndex = *SrcPtr;
+
+			for (int32 Index = 0; Index < SectionIndices.Num(); Index++)
+			{
+				uint32 VertIndex = *SrcPtr++;
+
+				bNeeds32BitIndicesOut |= (VertIndex > MAX_uint16);
+				Section.MinVertexIndex = FMath::Min<uint32>(VertIndex, Section.MinVertexIndex);
+				Section.MaxVertexIndex = FMath::Max<uint32>(VertIndex, Section.MaxVertexIndex);
+				*DestPtr++ = VertIndex;
+			}
+		}
+	}
+}
+
+
+/**
+ * If the StaticMesh has a valid HiRes SourceModel, this is the highest-resolution data available
+ * and it's what we want to build Nanite from. This function does just that. Generally this does
+ * the same steps as the loop over LODs in FStaticMeshBuilder::Build() below, however we skip
+ * various steps like reduction.
+ * 
+ * Note that this currently discards the returned Nanite-fractional-cut that NaniteBuilderModule.Build() returns.
+ * Current assumption is that if the HiRes SourceModel exists, then an authored lower-resolution
+ * LOD0 SourceModel already exists and should be used directly, instead of a fractional cut.
+ */
+static bool BuildNaniteFromHiResSourceModel(
+	UStaticMesh* StaticMesh, 
+	const FMeshNaniteSettings NaniteSettings, 
+	FBoxSphereBounds& HiResBoundsOut, 
+	Nanite::FResources& NaniteResourcesOut)
+{
+	if (ensure(StaticMesh->IsHiResMeshDescriptionValid()) == false)
+	{
+		return false;
+	}
+
+	TRACE_CPUPROFILER_EVENT_SCOPE(FStaticMeshBuilder::BuildNaniteFromHiResSourceModel);
+
+	FMeshDescription HiResMeshDescription = *StaticMesh->GetHiResMeshDescription();
+	FStaticMeshSourceModel& HiResSrcModel = StaticMesh->GetHiResSourceModel();
+	FMeshBuildSettings& HiResBuildSettings = HiResSrcModel.BuildSettings;		// cannot be const because FMeshDescriptionHelper modifies the LightmapIndex fields ?!?
+
+	// compute tangents, lightmap UVs, etc
+	FMeshDescriptionHelper MeshDescriptionHelper(&HiResBuildSettings);
+	MeshDescriptionHelper.SetupRenderMeshDescription(StaticMesh, HiResMeshDescription);
+
+	// make temporary RenderData storage that we can construct to pass to Nanite build
+	FStaticMeshRenderData HiResTempRenderData;
+	HiResTempRenderData.AllocateLODResources(1);
+	FStaticMeshLODResources& HiResStaticMeshLOD = HiResTempRenderData.LODResources[0];
+	HiResStaticMeshLOD.MaxDeviation = 0.0f;
+
+	//Prepare the PerSectionIndices array so we can optimize the index buffer for the GPU
+	//TODO: we do not need this for Nanite build, we only need CombinedIndices. Can we compute that directly?
+	TArray<TArray<uint32>> PerSectionIndices;
+	PerSectionIndices.AddDefaulted(HiResMeshDescription.PolygonGroups().Num());
+	HiResStaticMeshLOD.Sections.Empty(HiResMeshDescription.PolygonGroups().Num());
+
+	// Build the vertex and index buffer. We do not need WedgeMap or RemapVerts
+	TArray<int32> WedgeMap, RemapVerts;
+	TArray<FStaticMeshBuildVertex> StaticMeshBuildVertices;
+	BuildVertexBuffer(StaticMesh, HiResMeshDescription, HiResBuildSettings, WedgeMap, HiResStaticMeshLOD.Sections, PerSectionIndices, StaticMeshBuildVertices, MeshDescriptionHelper.GetOverlappingCorners(), RemapVerts);
+	WedgeMap.Empty();
+
+	const uint32 NumTextureCoord = HiResMeshDescription.VertexInstanceAttributes().GetAttributesRef<FVector2D>(MeshAttribute::VertexInstance::TextureCoordinate).GetNumChannels();
+
+	// Only the render data and vertex buffers will be used from now on, so we can discard Render MeshDescription
+	HiResMeshDescription.Empty();
+
+	// Concatenate the per-section index buffers.
+	TArray<uint32> CombinedIndices;
+	bool bNeeds32BitIndices = false;
+	BuildCombinedSectionIndices(PerSectionIndices, HiResStaticMeshLOD, CombinedIndices, bNeeds32BitIndices);
+
+	// compute bounds from the HiRes mesh before we do Nanite build because it will modify StaticMeshBuildVertices
+	ComputeBoundsFromVertexList(StaticMeshBuildVertices, HiResBoundsOut.Origin, HiResBoundsOut.BoxExtent, HiResBoundsOut.SphereRadius);
+
+	// run nanite build
+	{
+		TRACE_CPUPROFILER_EVENT_SCOPE(FStaticMeshBuilder::BuildNaniteFromHiResSourceModel::Nanite);
+		Nanite::IBuilderModule& NaniteBuilderModule = Nanite::IBuilderModule::Get();
+		if (!NaniteBuilderModule.Build(NaniteResourcesOut, StaticMeshBuildVertices, CombinedIndices, HiResStaticMeshLOD.Sections, NumTextureCoord, NaniteSettings))
+		{
+			UE_LOG(LogStaticMesh, Error, TEXT("Failed to build Nanite for HiRes static mesh. See previous line(s) for details."));
+			return false;
+		}
+	}
+
+	return true;
+}
+
+
+
 bool FStaticMeshBuilder::Build(FStaticMeshRenderData& StaticMeshRenderData, UStaticMesh* StaticMesh, const FStaticMeshLODGroup& LODGroup)
 {
-	FScopedSlowTask SlowTask(StaticMesh->GetNumSourceModels(), NSLOCTEXT("StaticMeshEditor", "StaticMeshBuilderBuild", "Building static mesh render data."));
+	const bool bNaniteBuildEnabled = StaticMesh->NaniteSettings.bEnabled;
+	const bool bHaveHiResSourceModel = StaticMesh->IsHiResMeshDescriptionValid();
+	int32 NumTasks = (bNaniteBuildEnabled && bHaveHiResSourceModel) ? (StaticMesh->GetNumSourceModels() + 1) : (StaticMesh->GetNumSourceModels());
+	FScopedSlowTask SlowTask(NumTasks, NSLOCTEXT("StaticMeshEditor", "StaticMeshBuilderBuild", "Building static mesh render data."));
 	SlowTask.MakeDialog();
 
 	// The tool can only been switch by restarting the editor
@@ -91,6 +272,25 @@ bool FStaticMeshBuilder::Build(FStaticMeshRenderData& StaticMeshRenderData, USta
 	const FMeshSectionInfoMap BeforeBuildOriginalSectionInfoMap = StaticMesh->GetOriginalSectionInfoMap();
 	const FMeshNaniteSettings NaniteSettings = StaticMesh->NaniteSettings;
 
+	bool bNaniteDataBuilt = false;		// true once we have finished building Nanite, which can happen in multiple places
+
+	// bounds of the pre-Nanite mesh
+	FBoxSphereBounds HiResBounds;
+	bool bHaveHiResBounds = false;
+
+	// do nanite build for HiRes SourceModel if we have one. In that case we skip the inline nanite build
+	// below that would happen with LOD0 build
+	if (bHaveHiResSourceModel && bNaniteBuildEnabled)
+	{
+		SlowTask.EnterProgressFrame(1);
+		if (BuildNaniteFromHiResSourceModel(StaticMesh, NaniteSettings, HiResBounds, StaticMeshRenderData.NaniteResources))
+		{
+			bHaveHiResBounds = true;
+			bNaniteDataBuilt = true;	// we are done building nanite data, and should not do it inline with LOD0 creation
+		}
+	}
+
+	// build render data for each LOD
 	for (int32 LodIndex = 0; LodIndex < NumSourceModels; ++LodIndex)
 	{
 		SlowTask.EnterProgressFrame(1);
@@ -160,7 +360,7 @@ bool FStaticMeshBuilder::Build(FStaticMeshRenderData& StaticMeshRenderData, USta
 
 		// Reduce LODs (if not building for Nanite).
 		// TODO: Skip all LOD reduction for Nanite?
-		if (bUseReduction && (!StaticMesh->NaniteSettings.bEnabled || LodIndex != 0))
+		if (bUseReduction && (!bNaniteBuildEnabled || LodIndex != 0))
 		{
 			float OverlappingThreshold = LODBuildSettings.bRemoveDegenerates ? THRESH_POINTS_ARE_SAME : 0.0f;
 			FOverlappingCorners OverlappingCorners;
@@ -290,43 +490,19 @@ bool FStaticMeshBuilder::Build(FStaticMeshRenderData& StaticMeshRenderData, USta
 		// Concatenate the per-section index buffers.
 		TArray<uint32> CombinedIndices;
 		bool bNeeds32BitIndices = false;
-		for (int32 SectionIndex = 0; SectionIndex < StaticMeshLOD.Sections.Num(); SectionIndex++)
-		{
-			FStaticMeshSection& Section = StaticMeshLOD.Sections[SectionIndex];
-			TArray<uint32> const& SectionIndices = PerSectionIndices[SectionIndex];
-			Section.FirstIndex = 0;
-			Section.NumTriangles = 0;
-			Section.MinVertexIndex = 0;
-			Section.MaxVertexIndex = 0;
-
-			if (SectionIndices.Num())
-			{
-				Section.FirstIndex = CombinedIndices.Num();
-				Section.NumTriangles = SectionIndices.Num() / 3;
-
-				CombinedIndices.AddUninitialized(SectionIndices.Num());
-				uint32* DestPtr = &CombinedIndices[Section.FirstIndex];
-				uint32 const* SrcPtr = SectionIndices.GetData();
-
-				Section.MinVertexIndex = *SrcPtr;
-				Section.MaxVertexIndex = *SrcPtr;
-
-				for (int32 Index = 0; Index < SectionIndices.Num(); Index++)
-				{
-					uint32 VertIndex = *SrcPtr++;
-
-					bNeeds32BitIndices |= (VertIndex > MAX_uint16);
-					Section.MinVertexIndex = FMath::Min<uint32>(VertIndex, Section.MinVertexIndex);
-					Section.MaxVertexIndex = FMath::Max<uint32>(VertIndex, Section.MaxVertexIndex);
-					*DestPtr++ = VertIndex;
-				}
-			}
-		}
+		BuildCombinedSectionIndices(PerSectionIndices, StaticMeshLOD, CombinedIndices, bNeeds32BitIndices);
 
-		if (StaticMesh->NaniteSettings.bEnabled && LodIndex == 0)
+		// If we want Nanite build, and have not already done it, do it based on LOD0 built render data.
+		// This will replace the output VertexBuffers/etc with the fractional Nanite cut to be stored as LOD0 RenderData.
+		// If we /have/ already done the Nanite build, it's because we have a separate HiRes mesh, and the LOD0 SourceModel
+		// has been created by some other means and should be used as the real LOD0
+		if (bNaniteBuildEnabled && LodIndex == 0 && bNaniteDataBuilt == false )
 		{
 			TRACE_CPUPROFILER_EVENT_SCOPE(FStaticMeshBuilder::Build::Nanite);
 
+			ComputeBoundsFromVertexList(StaticMeshBuildVertices, HiResBounds.Origin, HiResBounds.BoxExtent, HiResBounds.SphereRadius);
+			bHaveHiResBounds = true;
+
 			WedgeMap.Empty();	// Make sure to not keep the large WedgeMap from the input mesh around.
 								// No need to calculate a new one for the coarse mesh, because Nanite meshes don't need it yet.
 			Nanite::IBuilderModule& NaniteBuilderModule = Nanite::IBuilderModule::Get();
@@ -334,8 +510,11 @@ bool FStaticMeshBuilder::Build(FStaticMeshRenderData& StaticMeshRenderData, USta
 			{
 				UE_LOG(LogStaticMesh, Error, TEXT("Failed to build Nanite for static mesh. See previous line(s) for details."));
 			}
+
+			bNaniteDataBuilt = true;
 		}
 
+		// create the vertex and index buffers
 		{
 			TRACE_CPUPROFILER_EVENT_SCOPE(FStaticMeshBuilder::Build::BufferInit);
 			StaticMeshLOD.VertexBuffers.StaticMeshVertexBuffer.SetUseHighPrecisionTangentBasis(LODBuildSettings.bUseHighPrecisionTangentBasis);
@@ -347,28 +526,22 @@ bool FStaticMeshBuilder::Build(FStaticMeshRenderData& StaticMeshRenderData, USta
 			const EIndexBufferStride::Type IndexBufferStride = bNeeds32BitIndices ? EIndexBufferStride::Force32Bit : EIndexBufferStride::Force16Bit;
 			StaticMeshLOD.IndexBuffer.SetIndices(CombinedIndices, IndexBufferStride);
 		}
+
+		// post-process the index buffer
 		BuildLODSlowTask.EnterProgressFrame(1);
 		BuildAllBufferOptimizations(StaticMeshLOD, LODBuildSettings, CombinedIndices, bNeeds32BitIndices, StaticMeshBuildVertices);
+
 	} //End of LOD for loop
 
-	// Calculate the bounding box.
-	FBox BoundingBox(ForceInit);
+	// Calculate the bounding box of LOD0 buffer
 	FPositionVertexBuffer& BasePositionVertexBuffer = StaticMeshRenderData.LODResources[0].VertexBuffers.PositionVertexBuffer;
-	for (uint32 VertexIndex = 0; VertexIndex < BasePositionVertexBuffer.GetNumVertices(); VertexIndex++)
+	ComputeBoundsFromPositionBuffer(BasePositionVertexBuffer, StaticMeshRenderData.Bounds.Origin, StaticMeshRenderData.Bounds.BoxExtent, StaticMeshRenderData.Bounds.SphereRadius);
+	// combine with high-res bounds if it was computed
+	if (bHaveHiResBounds)
 	{
-		BoundingBox += BasePositionVertexBuffer.VertexPosition(VertexIndex);
+		StaticMeshRenderData.Bounds = StaticMeshRenderData.Bounds + HiResBounds;
 	}
-	BoundingBox.GetCenterAndExtents(StaticMeshRenderData.Bounds.Origin, StaticMeshRenderData.Bounds.BoxExtent);
 
-	// Calculate the bounding sphere, using the center of the bounding box as the origin.
-	StaticMeshRenderData.Bounds.SphereRadius = 0.0f;
-	for (uint32 VertexIndex = 0; VertexIndex < BasePositionVertexBuffer.GetNumVertices(); VertexIndex++)
-	{
-		StaticMeshRenderData.Bounds.SphereRadius = FMath::Max(
-			(BasePositionVertexBuffer.VertexPosition(VertexIndex) - StaticMeshRenderData.Bounds.Origin).Size(),
-			StaticMeshRenderData.Bounds.SphereRadius
-		);
-	}
 
 	return true;
 }
