Fix for Nanite building with reordered section material indices.
This involves a change to the DDC key format for static meshes with multiple sections in LOD0.

diff --git a/Engine/Source/Developer/MeshBuilder/Private/StaticMeshBuilder.cpp b/Engine/Source/Developer/MeshBuilder/Private/StaticMeshBuilder.cpp
--- a/Engine/Source/Developer/MeshBuilder/Private/StaticMeshBuilder.cpp
+++ b/Engine/Source/Developer/MeshBuilder/Private/StaticMeshBuilder.cpp
@@ -215,6 +215,13 @@ static bool BuildNaniteFromHiResSourceModel(
 	// compute bounds from the HiRes mesh before we do Nanite build because it will modify StaticMeshBuildVertices
 	ComputeBoundsFromVertexList(StaticMeshBuildVertices, HiResBoundsOut.Origin, HiResBoundsOut.BoxExtent, HiResBoundsOut.SphereRadius);
 
+	// Nanite build requires the section material indices to have already been resolved from the SectionInfoMap
+	// as the indices are baked into the FMaterialTriangles.
+	for (int32 SectionIndex = 0; SectionIndex < HiResStaticMeshLOD.Sections.Num(); SectionIndex++)
+	{
+		HiResStaticMeshLOD.Sections[SectionIndex].MaterialIndex = StaticMesh->GetSectionInfoMap().Get(0, SectionIndex).MaterialIndex;
+	}
+
 	// run nanite build
 	{
 		TRACE_CPUPROFILER_EVENT_SCOPE(FStaticMeshBuilder::BuildNaniteFromHiResSourceModel::Nanite);
@@ -503,6 +510,13 @@ bool FStaticMeshBuilder::Build(FStaticMeshRenderData& StaticMeshRenderData, USta
 			ComputeBoundsFromVertexList(StaticMeshBuildVertices, HiResBounds.Origin, HiResBounds.BoxExtent, HiResBounds.SphereRadius);
 			bHaveHiResBounds = true;
 
+			// Nanite build requires the section material indices to have already been resolved from the SectionInfoMap
+			// as the indices are baked into the FMaterialTriangles.
+			for (int32 SectionIndex = 0; SectionIndex < StaticMeshLOD.Sections.Num(); SectionIndex++)
+			{
+				StaticMeshLOD.Sections[SectionIndex].MaterialIndex = StaticMesh->GetSectionInfoMap().Get(LodIndex, SectionIndex).MaterialIndex;
+			}
+
 			WedgeMap.Empty();	// Make sure to not keep the large WedgeMap from the input mesh around.
 								// No need to calculate a new one for the coarse mesh, because Nanite meshes don't need it yet.
 			Nanite::IBuilderModule& NaniteBuilderModule = Nanite::IBuilderModule::Get();
diff --git a/Engine/Source/Runtime/Engine/Private/StaticMesh.cpp b/Engine/Source/Runtime/Engine/Private/StaticMesh.cpp
--- a/Engine/Source/Runtime/Engine/Private/StaticMesh.cpp
+++ b/Engine/Source/Runtime/Engine/Private/StaticMesh.cpp
@@ -2558,6 +2558,18 @@ static FString BuildStaticMeshDerivedDataKeySuffix(const ITargetPlatform* Target
 		KeySuffix += TEXT("_MinMLOD");
 	}
 
+	// Append the section material slot mappings for LOD0, if there are two or more (and thus are permutable), as they are baked into the Nanite build.
+	const FMeshSectionInfoMap& SectionInfoMap = Mesh->GetSectionInfoMap();
+	int32 NumLOD0Sections = SectionInfoMap.GetSectionNumber(0);
+	if (NumLOD0Sections >= 2)
+	{
+		KeySuffix += TEXT("_");
+		for (int32 SectionIndex = 0; SectionIndex < NumLOD0Sections; SectionIndex++)
+		{
+			KeySuffix += LexToString(SectionInfoMap.Get(0, SectionIndex).MaterialIndex);
+		}
+	}
+
 	return KeySuffix;
 }
 
