Make IMeshBuilderModule functions pure virtual and fix missing implementation on some platforms that was causing broken skeletal meshes in cooked builds.

diff --git a/Engine/Source/Developer/MeshBuilder/Public/IMeshBuilderModule.h b/Engine/Source/Developer/MeshBuilder/Public/IMeshBuilderModule.h
--- a/Engine/Source/Developer/MeshBuilder/Public/IMeshBuilderModule.h
+++ b/Engine/Source/Developer/MeshBuilder/Public/IMeshBuilderModule.h
@@ -29,16 +29,11 @@ public:
 
 	virtual void AppendToDDCKey(FString& DDCKey) { }
 
-	virtual bool BuildMesh(class FStaticMeshRenderData& OutRenderData, class UObject* Mesh, const class FStaticMeshLODGroup& LODGroup) { return false; }
+	virtual bool BuildMesh(class FStaticMeshRenderData& OutRenderData, class UObject* Mesh, const class FStaticMeshLODGroup& LODGroup) = 0;
 
-	virtual bool BuildMeshVertexPositions(
-		class UObject* StaticMesh,
-		TArray<uint32>& Indices,
-		TArray<FVector>& Vertices)
-	{
-		return false;
-	}
+	virtual bool BuildMeshVertexPositions(class UObject* StaticMesh, TArray<uint32>& Indices, TArray<FVector>& Vertices) = 0;
+
+	virtual bool BuildSkeletalMesh(const struct FSkeletalMeshBuildParameters& SkeletalMeshBuildParameters) = 0;
 
-	virtual bool BuildSkeletalMesh(const struct FSkeletalMeshBuildParameters& SkeletalMeshBuildParameters) { return false; }
 	virtual void PostBuildSkeletalMesh(class FSkeletalMeshRenderData* SkeletalMeshRenderData, class USkeletalMesh* SkeletalMesh) { }
 };
diff --git a/Engine/Source/Runtime/Engine/Private/SkeletalMeshRenderData.cpp b/Engine/Source/Runtime/Engine/Private/SkeletalMeshRenderData.cpp
--- a/Engine/Source/Runtime/Engine/Private/SkeletalMeshRenderData.cpp
+++ b/Engine/Source/Runtime/Engine/Private/SkeletalMeshRenderData.cpp
@@ -70,7 +70,7 @@ static void SerializeLODInfoForDDC(USkeletalMesh* SkeletalMesh, FString& KeySuff
 // differences, etc.) replace the version GUID below with a new one.
 // In case of merge conflicts with DDC versions, you MUST generate a new GUID
 // and set this new GUID as the version.
-#define SKELETALMESH_DERIVEDDATA_VER TEXT("EC1BEA18C73F42AE95D70CDC66EBB51A")
+#define SKELETALMESH_DERIVEDDATA_VER TEXT("23BBACC58782442DAAE9DCBFC4715C65")
 
 const FString& GetSkeletalMeshDerivedDataVersion()
 {
