Fix build for Max & Linus platforms

diff --git a/Engine/Source/Developer/MeshBuilder/Private/StaticMeshBuilder.cpp b/Engine/Source/Developer/MeshBuilder/Private/StaticMeshBuilder.cpp
--- a/Engine/Source/Developer/MeshBuilder/Private/StaticMeshBuilder.cpp
+++ b/Engine/Source/Developer/MeshBuilder/Private/StaticMeshBuilder.cpp
@@ -575,8 +575,11 @@ bool FStaticMeshBuilder::Build(FStaticMeshRenderData& StaticMeshRenderData, USta
 					}
 
 					FStaticMeshLODResources& ProxyLOD = StaticMeshRenderData.LODResources[ValidLODCount];
-
-					ProxyLOD.Sections = FStaticMeshSectionArray(ProxyMeshData.Sections);
+					ProxyLOD.Sections.Empty(ProxyMeshData.Sections.Num());
+					for (FStaticMeshSection& Section : ProxyMeshData.Sections)
+					{
+						ProxyLOD.Sections.Add(Section);
+					}
 
 					TRACE_CPUPROFILER_EVENT_SCOPE(FStaticMeshBuilder::Build::BufferInit);
 					ProxyLOD.VertexBuffers.StaticMeshVertexBuffer.SetUseHighPrecisionTangentBasis(LODBuildSettings.bUseHighPrecisionTangentBasis);
