StaticMesh: update Nanite Build from HiRes SourceModel that had regressed relative to changes in Nanite builder
#rnx

diff --git a/Engine/Source/Developer/MeshBuilder/Private/StaticMeshBuilder.cpp b/Engine/Source/Developer/MeshBuilder/Private/StaticMeshBuilder.cpp
--- a/Engine/Source/Developer/MeshBuilder/Private/StaticMeshBuilder.cpp
+++ b/Engine/Source/Developer/MeshBuilder/Private/StaticMeshBuilder.cpp
@@ -221,12 +221,17 @@ static bool BuildNaniteFromHiResSourceModel(
 		HiResStaticMeshLOD.Sections[SectionIndex].MaterialIndex = StaticMesh->GetSectionInfoMap().Get(0, SectionIndex).MaterialIndex;
 	}
 
+	// need to copy the sections list
+	InputMeshData.Sections = HiResStaticMeshLOD.Sections;
+
 	// run nanite build
 	{
 		TRACE_CPUPROFILER_EVENT_SCOPE(FStaticMeshBuilder::BuildNaniteFromHiResSourceModel::Nanite);
 		Nanite::IBuilderModule& NaniteBuilderModule = Nanite::IBuilderModule::Get();
 		TArray<Nanite::IBuilderModule::FVertexMeshData, TInlineAllocator<2>> EmptyOutputLODMeshData; //< not output mesh data required
-		if (!NaniteBuilderModule.Build(NaniteResourcesOut, InputMeshData, EmptyOutputLODMeshData, NumTextureCoord, NaniteSettings))
+		FMeshNaniteSettings LocalNaniteSettings = NaniteSettings;
+		LocalNaniteSettings.PercentTriangles = 1.0f;		// prevent Nanite from trying to generate coarse representation / proxy mesh
+		if (!NaniteBuilderModule.Build(NaniteResourcesOut, InputMeshData, EmptyOutputLODMeshData, NumTextureCoord, LocalNaniteSettings))
 		{
 			UE_LOG(LogStaticMesh, Error, TEXT("Failed to build Nanite for HiRes static mesh. See previous line(s) for details."));
 			return false;
