Move dependency to public to fix compile errors in other modules that include headers
#rnx

diff --git a/Engine/Source/Developer/MeshBuilder/MeshBuilder.Build.cs b/Engine/Source/Developer/MeshBuilder/MeshBuilder.Build.cs
--- a/Engine/Source/Developer/MeshBuilder/MeshBuilder.Build.cs
+++ b/Engine/Source/Developer/MeshBuilder/MeshBuilder.Build.cs
@@ -7,6 +7,13 @@ namespace UnrealBuildTool.Rules
 	{
 		public MeshBuilder(ReadOnlyTargetRules Target) : base(Target)
 		{
+			PublicDependencyModuleNames.AddRange(
+				new string[]
+				{
+					"MeshUtilitiesCommon",
+				}
+			);
+
 			PrivateDependencyModuleNames.AddRange(
 				new string[]
 				{
@@ -20,7 +27,6 @@ namespace UnrealBuildTool.Rules
                     "MeshReductionInterface",
                     "RawMesh",
 					"MeshUtilities",
-					"MeshUtilitiesCommon",
 					"ClothingSystemRuntimeNv",
 					"MeshBoneReduction",
 					"SkeletalMeshUtilitiesCommon",
