StaticMeshBuilder.h - NoPCH build fix

diff --git a/Engine/Source/Developer/MeshBuilder/Public/StaticMeshBuilder.h b/Engine/Source/Developer/MeshBuilder/Public/StaticMeshBuilder.h
--- a/Engine/Source/Developer/MeshBuilder/Public/StaticMeshBuilder.h
+++ b/Engine/Source/Developer/MeshBuilder/Public/StaticMeshBuilder.h
@@ -12,6 +12,9 @@ class FStaticMeshRenderData;
 class FStaticMeshLODGroup;
 class USkeletalMesh;
 struct FOverlappingCorners;
+struct FMeshDescription;
+struct FMeshBuildSettings;
+
 
 class MESHBUILDER_API FStaticMeshBuilder : public FMeshBuilder
 {
