Fix a bug where high resolution Nanite meshes would not be simplified on platforms that do not support Nanite

diff --git a/Engine/Source/Developer/MeshBuilder/Private/StaticMeshBuilder.cpp b/Engine/Source/Developer/MeshBuilder/Private/StaticMeshBuilder.cpp
--- a/Engine/Source/Developer/MeshBuilder/Private/StaticMeshBuilder.cpp
+++ b/Engine/Source/Developer/MeshBuilder/Private/StaticMeshBuilder.cpp
@@ -262,9 +262,9 @@ static bool BuildNanite(
 }
 
 
-bool FStaticMeshBuilder::Build(FStaticMeshRenderData& StaticMeshRenderData, UStaticMesh* StaticMesh, const FStaticMeshLODGroup& LODGroup, bool bGenerateCoarseMeshStreamingLODs, bool bAllowNanite)
+bool FStaticMeshBuilder::Build(FStaticMeshRenderData& StaticMeshRenderData, UStaticMesh* StaticMesh, const FStaticMeshLODGroup& LODGroup, bool bGenerateCoarseMeshStreamingLODs, bool bTargetSupportsNanite)
 {
-	const bool bNaniteBuildEnabled = StaticMesh->NaniteSettings.bEnabled && bAllowNanite;
+	const bool bNaniteBuildEnabled = StaticMesh->NaniteSettings.bEnabled;
 	const bool bHaveHiResSourceModel = StaticMesh->IsHiResMeshDescriptionValid();
 	int32 NumTasks = (bNaniteBuildEnabled && bHaveHiResSourceModel) ? (StaticMesh->GetNumSourceModels() + 1) : (StaticMesh->GetNumSourceModels());
 	FScopedSlowTask SlowTask(NumTasks, NSLOCTEXT("StaticMeshEditor", "StaticMeshBuilderBuild", "Building static mesh render data."));
@@ -312,7 +312,7 @@ bool FStaticMeshBuilder::Build(FStaticMeshRenderData& StaticMeshRenderData, USta
 
 	// Do nanite build for HiRes SourceModel if we have one. In that case we skip the inline nanite build
 	// below that would happen with LOD0 build
-	if (bHaveHiResSourceModel && bNaniteBuildEnabled)
+	if (bHaveHiResSourceModel && bNaniteBuildEnabled && bTargetSupportsNanite)
 	{
 		SlowTask.EnterProgressFrame(1);
 
@@ -335,6 +335,8 @@ bool FStaticMeshBuilder::Build(FStaticMeshRenderData& StaticMeshRenderData, USta
 
 	// If we want Nanite built, and have not already done it, do it based on LOD0 built render data.
 	// This will replace the output VertexBuffers/etc with the fractional Nanite cut to be stored as LOD0 RenderData.
+	// NOTE: We still want to do this for targets that do not support Nanite so that it generates the fallback,
+	// in which case the Nanite bulk will be stripped
 	if (!bNaniteDataBuilt && bNaniteBuildEnabled)
 	{
 		TArray< float, TInlineAllocator<4> > PercentTriangles;
@@ -366,6 +368,12 @@ bool FStaticMeshBuilder::Build(FStaticMeshRenderData& StaticMeshRenderData, USta
 		bHaveHiResBounds = true;
 		bNaniteDataBuilt = true;
 		NaniteBuiltLevels = PercentTriangles.Num();
+
+		if (!bTargetSupportsNanite)
+		{
+			// Strip the Nanite bulk for the target platform
+			StaticMeshRenderData.NaniteResources = Nanite::FResources();
+		}
 	}
 
 	// Build render data for each LOD, starting from where Nanite left off.
@@ -392,7 +400,7 @@ bool FStaticMeshBuilder::Build(FStaticMeshRenderData& StaticMeshRenderData, USta
 
 		if (bIsMeshDescriptionValid)
 		{
-			MeshDescriptionHelper.SetupRenderMeshDescription(StaticMesh, MeshDescriptions[LodIndex], bAllowNanite);
+			MeshDescriptionHelper.SetupRenderMeshDescription(StaticMesh, MeshDescriptions[LodIndex], bTargetSupportsNanite);
 		}
 		else
 		{
@@ -566,7 +574,7 @@ bool FStaticMeshBuilder::Build(FStaticMeshRenderData& StaticMeshRenderData, USta
 			StaticMeshBuildVertices,
 			MeshDescriptionHelper.GetOverlappingCorners(),
 			RemapVerts,
-			bAllowNanite
+			bTargetSupportsNanite
 		);
 
 		TVertexInstanceAttributesRef<FVector2f> VertexInstanceUVs = MeshDescriptions[LodIndex].VertexInstanceAttributes().GetAttributesRef<FVector2f>(MeshAttribute::VertexInstance::TextureCoordinate);
