The spikes seems to not happen any more with the new simplification algorithm.
Whatever, the bounding box of the simplified mesh is compared to the original mesh. If the new bounding box is larger, a message is output.

diff --git a/Engine/Source/Developer/MeshBuilder/Private/StaticMeshBuilder.cpp b/Engine/Source/Developer/MeshBuilder/Private/StaticMeshBuilder.cpp
--- a/Engine/Source/Developer/MeshBuilder/Private/StaticMeshBuilder.cpp
+++ b/Engine/Source/Developer/MeshBuilder/Private/StaticMeshBuilder.cpp
@@ -462,15 +462,31 @@ bool FStaticMeshBuilder::Build(FStaticMeshRenderData& StaticMeshRenderData, USta
 
 			int32 OldSectionInfoMapCount = StaticMesh->GetSectionInfoMap().GetSectionNumber(LodIndex);
 
+			TFunction<void(const FMeshDescription&, const FMeshDescription&)> CheckReduction = [&](const FMeshDescription& InitMesh, const FMeshDescription& ReducedMesh)
+			{
+				FBox BBoxInitMesh = InitMesh.ComputeBoundingBox();
+				double BBoxInitMeshSize = (BBoxInitMesh.Max - BBoxInitMesh.Min).Length();
+
+				FBox BBoxReducedMesh = ReducedMesh.ComputeBoundingBox();
+				double BBoxReducedMeshSize = (BBoxReducedMesh.Max - BBoxReducedMesh.Min).Length();
+
+				if (BBoxReducedMeshSize > BBoxInitMeshSize * 1.01)
+				{
+					UE_LOG(LogStaticMeshBuilder, Warning, TEXT("The generation of LOD could have generated spikes on the mesh for %s"), *StaticMesh->GetName());
+				}
+			};
+
 			if (LodIndex == BaseReduceLodIndex)
 			{
 				//When using LOD 0, we use a copy of the mesh description since reduce do not support inline reducing
 				FMeshDescription BaseMeshDescription = MeshDescriptions[BaseReduceLodIndex];
 				MeshDescriptionHelper.ReduceLOD(BaseMeshDescription, MeshDescriptions[LodIndex], ReductionSettings, OverlappingCorners, MaxDeviation);
+				CheckReduction(BaseMeshDescription, MeshDescriptions[LodIndex]);
 			}
 			else
 			{
 				MeshDescriptionHelper.ReduceLOD(MeshDescriptions[BaseReduceLodIndex], MeshDescriptions[LodIndex], ReductionSettings, OverlappingCorners, MaxDeviation);
+				CheckReduction(MeshDescriptions[BaseReduceLodIndex], MeshDescriptions[LodIndex]);
 			}
 			
 
