Distance Fields - Generate mesh distance fields from highest quality available data in order to match Nanite data.

diff --git a/Engine/Source/Developer/MeshBuilder/Private/StaticMeshBuilder.cpp b/Engine/Source/Developer/MeshBuilder/Private/StaticMeshBuilder.cpp
--- a/Engine/Source/Developer/MeshBuilder/Private/StaticMeshBuilder.cpp
+++ b/Engine/Source/Developer/MeshBuilder/Private/StaticMeshBuilder.cpp
@@ -659,7 +659,8 @@ bool FStaticMeshBuilder::BuildMeshVertexPositions(
 {
 	TRACE_CPUPROFILER_EVENT_SCOPE(FStaticMeshBuilder::BuildMeshVertexPositions);
 
-	if (!StaticMesh->IsMeshDescriptionValid(0))
+	FStaticMeshSourceModel& SourceModel = StaticMesh->IsHiResMeshDescriptionValid() ? StaticMesh->GetHiResSourceModel() : StaticMesh->GetSourceModel(0);
+	if (!SourceModel.IsMeshDescriptionValid())
 	{
 		//Warn the user that there is no mesh description data
 		UE_LOG(LogStaticMeshBuilder, Error, TEXT("Cannot find a valid mesh description to build the asset."));
@@ -670,10 +671,10 @@ bool FStaticMeshBuilder::BuildMeshVertexPositions(
 	if (NumSourceModels > 0)
 	{
 		FMeshDescription MeshDescription;
-		const bool bIsMeshDescriptionValid = StaticMesh->CloneMeshDescription(/*LodIndex*/ 0, MeshDescription);
+		const bool bIsMeshDescriptionValid = SourceModel.CloneMeshDescription(MeshDescription);
 		if (bIsMeshDescriptionValid)
 		{
-			const FMeshBuildSettings& BuildSettings = StaticMesh->GetSourceModel(0).BuildSettings;
+			const FMeshBuildSettings& BuildSettings = SourceModel.BuildSettings;
 
 			const FStaticMeshConstAttributes Attributes(MeshDescription);
 			TArrayView<const FVector3f> VertexPositions = Attributes.GetVertexPositions().GetRawArray();
diff --git a/Engine/Source/Runtime/Engine/Private/DistanceFieldAtlas.cpp b/Engine/Source/Runtime/Engine/Private/DistanceFieldAtlas.cpp
--- a/Engine/Source/Runtime/Engine/Private/DistanceFieldAtlas.cpp
+++ b/Engine/Source/Runtime/Engine/Private/DistanceFieldAtlas.cpp
@@ -132,7 +132,7 @@ FDistanceFieldAsyncQueue* GDistanceFieldAsyncQueue = NULL;
 #if WITH_EDITOR
 
 // DDC key for distance field data, must be changed when modifying the generation code or data format
-#define DISTANCEFIELD_DERIVEDDATA_VER TEXT("23A632E6-DB31-4A5B-B43D-17D98C6E946C")
+#define DISTANCEFIELD_DERIVEDDATA_VER TEXT("8A773878-E42B-4143-A23B-E231679CE486")
 
 FString BuildDistanceFieldDerivedDataKey(const FString& InMeshKey)
 {
