Fix an issue where Nanite would be built for LOD0 when it was not supposed to be for non-Nanite platforms.

#rnx

diff --git a/Engine/Source/Developer/MeshBuilder/Private/StaticMeshBuilder.cpp b/Engine/Source/Developer/MeshBuilder/Private/StaticMeshBuilder.cpp
--- a/Engine/Source/Developer/MeshBuilder/Private/StaticMeshBuilder.cpp
+++ b/Engine/Source/Developer/MeshBuilder/Private/StaticMeshBuilder.cpp
@@ -341,9 +341,9 @@ bool FStaticMeshBuilder::Build(FStaticMeshRenderData& StaticMeshRenderData, USta
 
 	// If we want Nanite built, and have not already done it, do it based on LOD0 built render data.
 	// This will replace the output VertexBuffers/etc with the fractional Nanite cut to be stored as LOD0 RenderData.
-	// NOTE: We still want to do this for targets that do not support Nanite so that it generates the fallback,
-	// in which case the Nanite bulk will be stripped
-	if (!bNaniteDataBuilt && bNaniteBuildEnabled)
+	// NOTE: We still want to do this for targets that do not support Nanite (if no hi-res source model was provided)
+	// so that it generates the fallback, in which case the Nanite bulk will be stripped
+	if (bNaniteBuildEnabled && ((bTargetSupportsNanite && !bNaniteDataBuilt) || (!bTargetSupportsNanite && !bHaveHiResSourceModel)))
 	{
 		TArray< float, TInlineAllocator<4> > PercentTriangles;
 		for (int32 LodIndex = 0; LodIndex < NumSourceModels; ++LodIndex)
diff --git a/Engine/Source/Runtime/Core/Private/UObject/FortniteMainBranchObjectVersion.cpp b/Engine/Source/Runtime/Core/Private/UObject/FortniteMainBranchObjectVersion.cpp
--- a/Engine/Source/Runtime/Core/Private/UObject/FortniteMainBranchObjectVersion.cpp
+++ b/Engine/Source/Runtime/Core/Private/UObject/FortniteMainBranchObjectVersion.cpp
@@ -13,6 +13,6 @@ TMap<FGuid, FGuid> FFortniteMainBranchObjectVersion::GetSystemGuids()
 	SystemGuids.Add(DevGuids.NIAGARASHADERMAP_DERIVEDDATA_VER, FGuid("2289D5116CF94BC0AFEAEC541468E645"));
 	SystemGuids.Add(DevGuids.Niagara_LatestScriptCompileVersion, FGuid("9CFE679BEE95F345BBC06744597AC064"));
 	SystemGuids.Add(DevGuids.SkeletalMeshDerivedDataVersion, FGuid("ACF593EAAE354FCBB34CF44F5AA1BFD2"));
-	SystemGuids.Add(DevGuids.STATICMESH_DERIVEDDATA_VER, FGuid("18E028FCF1BC469385806AA5F8AE07CC"));
+	SystemGuids.Add(DevGuids.STATICMESH_DERIVEDDATA_VER, FGuid("F3F4E2FE7EE14CE7942B445FABD05E37"));
 	return SystemGuids;
 }
