The threshold in LOD generation to check the quality of the LOD was to small.
It has been update.

diff --git a/Engine/Source/Developer/MeshBuilder/Private/StaticMeshBuilder.cpp b/Engine/Source/Developer/MeshBuilder/Private/StaticMeshBuilder.cpp
--- a/Engine/Source/Developer/MeshBuilder/Private/StaticMeshBuilder.cpp
+++ b/Engine/Source/Developer/MeshBuilder/Private/StaticMeshBuilder.cpp
@@ -470,7 +470,8 @@ bool FStaticMeshBuilder::Build(FStaticMeshRenderData& StaticMeshRenderData, USta
 				FBox BBoxReducedMesh = ReducedMesh.ComputeBoundingBox();
 				double BBoxReducedMeshSize = (BBoxReducedMesh.Max - BBoxReducedMesh.Min).Length();
 
-				if (BBoxReducedMeshSize > BBoxInitMeshSize * 1.01)
+				constexpr double ThresholdForAbnormalGrowthOfBBox = 1.1;
+				if (BBoxReducedMeshSize > BBoxInitMeshSize * ThresholdForAbnormalGrowthOfBBox)
 				{
 					UE_LOG(LogStaticMeshBuilder, Warning, TEXT("The generation of LOD could have generated spikes on the mesh for %s"), *StaticMesh->GetName());
 				}
