The threshold in LOD generation to check the quality of the LOD was to small.
It has been update to sqrt(3) i.e. the reduced mesh must stay in the bounding sphere.

diff --git a/Engine/Source/Developer/MeshBuilder/Private/StaticMeshBuilder.cpp b/Engine/Source/Developer/MeshBuilder/Private/StaticMeshBuilder.cpp
--- a/Engine/Source/Developer/MeshBuilder/Private/StaticMeshBuilder.cpp
+++ b/Engine/Source/Developer/MeshBuilder/Private/StaticMeshBuilder.cpp
@@ -470,7 +470,7 @@ bool FStaticMeshBuilder::Build(FStaticMeshRenderData& StaticMeshRenderData, USta
 				FBox BBoxReducedMesh = ReducedMesh.ComputeBoundingBox();
 				double BBoxReducedMeshSize = (BBoxReducedMesh.Max - BBoxReducedMesh.Min).Length();
 
-				constexpr double ThresholdForAbnormalGrowthOfBBox = 1.1;
+				constexpr double ThresholdForAbnormalGrowthOfBBox = UE_DOUBLE_SQRT_3; // the reduced mesh must stay in the bounding sphere 
 				if (BBoxReducedMeshSize > BBoxInitMeshSize * ThresholdForAbnormalGrowthOfBBox)
 				{
 					UE_LOG(LogStaticMeshBuilder, Warning, TEXT("The generation of LOD could have generated spikes on the mesh for %s"), *StaticMesh->GetName());
