Fixed issue where non-compact mesh descriptions can build with unused extra data.

diff --git a/Engine/Source/Developer/MeshBuilder/Private/StaticMeshBuilder.cpp b/Engine/Source/Developer/MeshBuilder/Private/StaticMeshBuilder.cpp
--- a/Engine/Source/Developer/MeshBuilder/Private/StaticMeshBuilder.cpp
+++ b/Engine/Source/Developer/MeshBuilder/Private/StaticMeshBuilder.cpp
@@ -686,6 +686,9 @@ bool FStaticMeshBuilder::BuildMeshVertexPositions(
 		const bool bIsMeshDescriptionValid = SourceModel.CloneMeshDescription(MeshDescription);
 		if (bIsMeshDescriptionValid)
 		{
+			FElementIDRemappings Remappings;
+			MeshDescription.Compact(Remappings);
+
 			const FMeshBuildSettings& BuildSettings = SourceModel.BuildSettings;
 
 			const FStaticMeshConstAttributes Attributes(MeshDescription);
