Another non-unity fix

diff --git a/Engine/Source/Developer/MeshBuilder/Public/MeshBuilder.h b/Engine/Source/Developer/MeshBuilder/Public/MeshBuilder.h
--- a/Engine/Source/Developer/MeshBuilder/Public/MeshBuilder.h
+++ b/Engine/Source/Developer/MeshBuilder/Public/MeshBuilder.h
@@ -14,6 +14,7 @@ class UStaticMesh;
 struct FSkeletalMeshBuildParameters;
 struct FStaticMeshBuildVertex;
 struct FStaticMeshSection;
+struct FMeshBuildVertexData;
 
 /**
  * Abstract class which is the base class of all builder.
