Possible fix for rare crash in StaticMeshBuilder when assigning accessing mesh description

diff --git a/Engine/Source/Developer/MeshBuilder/Private/StaticMeshBuilder.cpp b/Engine/Source/Developer/MeshBuilder/Private/StaticMeshBuilder.cpp
--- a/Engine/Source/Developer/MeshBuilder/Private/StaticMeshBuilder.cpp
+++ b/Engine/Source/Developer/MeshBuilder/Private/StaticMeshBuilder.cpp
@@ -209,7 +209,7 @@ static bool BuildNanite(
 
 	// Only the render data and vertex buffers will be used from now on unless we have more than one source models
 	// This will help with memory usage for Nanite Mesh by releasing memory before doing the build
-	SourceModel.ClearMeshDescription(); // Clear original cached mesh description in memory
+	//SourceModel.ClearMeshDescription(); // Clear original cached mesh description in memory
 	MutableMeshDescription.Empty(); // Clear the mutable copy (if allocated).
 
 	// Concatenate the per-section index buffers.
