Fixed issue where Nanite fallback meshes with LODIndex>0 would be extremely low poly because PercentTriangles was uninitialized.
The initialization code for PercentTriangles was lost in 24853395, seemingly by mistake.

diff --git a/Engine/Source/Developer/MeshBuilder/Private/StaticMeshBuilder.cpp b/Engine/Source/Developer/MeshBuilder/Private/StaticMeshBuilder.cpp
--- a/Engine/Source/Developer/MeshBuilder/Private/StaticMeshBuilder.cpp
+++ b/Engine/Source/Developer/MeshBuilder/Private/StaticMeshBuilder.cpp
@@ -218,6 +218,11 @@ static bool BuildNanite(
 	TArray< Nanite::IBuilderModule::FOutputMeshData, TInlineAllocator<4> > OutputLODMeshData;
 	OutputLODMeshData.SetNum( PercentTriangles.Num() );
 
+	for (int32 LodIndex = 0; LodIndex < OutputLODMeshData.Num(); LodIndex++)
+	{
+		OutputLODMeshData[LodIndex].PercentTriangles = PercentTriangles[LodIndex];
+	}
+
 	const uint32 TriangleCount = InputMeshData.TriangleIndices.Num() / 3;
 
 	if (!NaniteBuilderModule.BuildMaterialIndices(InputMeshData.Sections, TriangleCount, InputMeshData.MaterialIndices))
