Added two missing error messages to BuildNanite that now prevent a failed build from continuing (and producing broken data and/or various ensures in the slow task dialog).

diff --git a/Engine/Source/Developer/MeshBuilder/Private/StaticMeshBuilder.cpp b/Engine/Source/Developer/MeshBuilder/Private/StaticMeshBuilder.cpp
--- a/Engine/Source/Developer/MeshBuilder/Private/StaticMeshBuilder.cpp
+++ b/Engine/Source/Developer/MeshBuilder/Private/StaticMeshBuilder.cpp
@@ -131,14 +131,16 @@ static bool BuildNanite(
 {
 	TRACE_CPUPROFILER_EVENT_SCOPE( FStaticMeshBuilder::BuildNanite );
 
-	if( !SourceModel.IsMeshDescriptionValid() )
+	if (!SourceModel.IsMeshDescriptionValid())
 	{
+		UE_LOG(LogStaticMeshBuilder, Error, TEXT("Invalid mesh description during Nanite build [%s]."), *StaticMesh->GetFullName());
 		return false;
 	}
 
 	FMeshDescription MeshDescription; 
 	if (!SourceModel.CloneMeshDescription(MeshDescription))
 	{
+		UE_LOG(LogStaticMeshBuilder, Error, TEXT("Failed to clone mesh description during Nanite build [%s]."), *StaticMesh->GetFullName());
 		return false;
 	}
 
