Adjusted error localisation keys to be more specific (and not duplicated)

diff --git a/Engine/Source/Developer/MeshBuilder/Private/StaticMeshBuilder.cpp b/Engine/Source/Developer/MeshBuilder/Private/StaticMeshBuilder.cpp
--- a/Engine/Source/Developer/MeshBuilder/Private/StaticMeshBuilder.cpp
+++ b/Engine/Source/Developer/MeshBuilder/Private/StaticMeshBuilder.cpp
@@ -235,7 +235,7 @@ static bool BuildNanite(
 
 	if (!NaniteBuilderModule.BuildMaterialIndices(InputMeshData.Sections, TriangleCount, InputMeshData.MaterialIndices))
 	{
-		UE_LOGFMT_NSLOC(LogStaticMesh, Error, "StaticMesh", "Error", "Failed to build Nanite from static mesh. See previous line(s) for details.");
+		UE_LOGFMT_NSLOC(LogStaticMesh, Error, "StaticMesh", "NaniteBuildError", "Failed to build Nanite from static mesh. See previous line(s) for details.");
 		return false;
 	}
 
@@ -260,7 +260,7 @@ static bool BuildNanite(
 			OnFreeInputMeshData
 	))
 	{
-		UE_LOGFMT_NSLOC(LogStaticMesh, Error, "StaticMesh", "Error", "Failed to build Nanite for HiRes static mesh. See previous line(s) for details.");
+		UE_LOGFMT_NSLOC(LogStaticMesh, Error, "StaticMesh", "NaniteHiResBuildError", "Failed to build Nanite for HiRes static mesh. See previous line(s) for details.");
 		return false;
 	}
 
