Make sure Hires mesh number of sections match LOD 0 number of sections.
#rnx

diff --git a/Engine/Source/Developer/MeshBuilder/Private/StaticMeshBuilder.cpp b/Engine/Source/Developer/MeshBuilder/Private/StaticMeshBuilder.cpp
--- a/Engine/Source/Developer/MeshBuilder/Private/StaticMeshBuilder.cpp
+++ b/Engine/Source/Developer/MeshBuilder/Private/StaticMeshBuilder.cpp
@@ -381,22 +381,46 @@ bool FStaticMeshBuilder::Build(FStaticMeshRenderData& StaticMeshRenderData, USta
 	{
 		SlowTask.EnterProgressFrame(1);
 
-		FBoxSphereBounds NaniteBounds;
-
-		bool bBuildSuccess = BuildNanite(
-			StaticMesh,
-			StaticMesh->GetHiResSourceModel(),
-			StaticMeshRenderData.LODResources,
-			StaticMeshRenderData.LODVertexFactories,
-			NaniteResources,
-			NaniteSettings,
-			TArrayView< float >(),
-			NaniteBounds);
+		auto IsHiresMeshDescriptionValid = [&StaticMesh]()
+			{
+				if (const FMeshDescription* BaseLodMeshDescription = StaticMesh->GetSourceModel(0).GetOrCacheMeshDescription())
+				{
+					if (const FMeshDescription* HiResMeshDescription = StaticMesh->GetHiResSourceModel().GetOrCacheMeshDescription())
+					{
+						//Validate the number of sections
+						if (HiResMeshDescription->PolygonGroups().Num() == BaseLodMeshDescription->PolygonGroups().Num())
+						{
+							return true;
+						}
+						else
+						{
+							UE_LOG(LogStaticMeshBuilder, Error, TEXT("Invalid hi-res mesh description during Nanite build [%s]. The number of sections from the hires mesh differ from the LOD 0 mesh, this is not supported and LOD 0 will be use as a fallback to build nanite data."), *StaticMesh->GetFullName());
+						}
+					}
+				}
+				//No need to log if we miss a mesh description, this was handle before
+				return false;
+			};
 
-		if( bBuildSuccess )
+		//Make sure hires mesh data has the same amount of sections. If not rendering bugs and issues will show up because the nanite render must use the LOD 0 sections.
+		if (IsHiresMeshDescriptionValid())
 		{
-			MeshBoundsBuilder += NaniteBounds;
-			bNaniteDataBuilt = true;
+			FBoxSphereBounds NaniteBounds;
+			bool bBuildSuccess = BuildNanite(
+				StaticMesh,
+				StaticMesh->GetHiResSourceModel(),
+				StaticMeshRenderData.LODResources,
+				StaticMeshRenderData.LODVertexFactories,
+				NaniteResources,
+				NaniteSettings,
+				TArrayView< float >(),
+				NaniteBounds);
+
+			if (bBuildSuccess)
+			{
+				MeshBoundsBuilder += NaniteBounds;
+				bNaniteDataBuilt = true;
+			}
 		}
 	}
 
