Make sure Hires mesh number of sections is equal or lower than LOD 0 number of sections. Issue an Error log if the hires is invalid and default nanite on the LOD 0.
#rnx

diff --git a/Engine/Source/Developer/MeshBuilder/Private/StaticMeshBuilder.cpp b/Engine/Source/Developer/MeshBuilder/Private/StaticMeshBuilder.cpp
--- a/Engine/Source/Developer/MeshBuilder/Private/StaticMeshBuilder.cpp
+++ b/Engine/Source/Developer/MeshBuilder/Private/StaticMeshBuilder.cpp
@@ -381,22 +381,52 @@ bool FStaticMeshBuilder::Build(FStaticMeshRenderData& StaticMeshRenderData, USta
 	{
 		SlowTask.EnterProgressFrame(1);
 
-		FBoxSphereBounds NaniteBounds;
-
-		bool bBuildSuccess = BuildNanite(
-			StaticMesh,
-			StaticMesh->GetHiResSourceModel(),
-			StaticMeshRenderData.LODResources,
-			StaticMeshRenderData.LODVertexFactories,
-			NaniteResources,
-			NaniteSettings,
-			TArrayView< float >(),
-			NaniteBounds);
+		auto IsHiresMeshDescriptionValid = [&StaticMesh]()
+			{
+				bool bIsValid = false;
+				if (const FMeshDescription* BaseLodMeshDescription = StaticMesh->GetSourceModel(0).GetOrCacheMeshDescription())
+				{
+					if (const FMeshDescription* HiResMeshDescription = StaticMesh->GetHiResSourceModel().GetOrCacheMeshDescription())
+					{
+						//Validate the number of sections
+						if (HiResMeshDescription->PolygonGroups().Num() > BaseLodMeshDescription->PolygonGroups().Num())
+						{
+							UE_LOG(LogStaticMeshBuilder, Error, TEXT("Invalid hi-res mesh description during Nanite build [%s]. The number of sections from the hires mesh is higher than LOD 0 section count. This is not supported and LOD 0 will be used as a fallback to build nanite data."), *StaticMesh->GetFullName());
+							bIsValid = false;
+						}
+						else
+						{
+							if (HiResMeshDescription->PolygonGroups().Num() < BaseLodMeshDescription->PolygonGroups().Num())
+							{
+								UE_LOG(LogStaticMeshBuilder, Display, TEXT("Nanite hi-res mesh description for [%s] has fewer sections than lod 0. Verify you have the proper material id result when nanite is turned on."), *StaticMesh->GetFullName());
+							}
+							bIsValid = true;
+						}
+					}
+				}
+				//No need to log if we miss a mesh description, this was handle before
+				return bIsValid;
+			};
 
-		if( bBuildSuccess )
+		//Make sure hires mesh data has the same amount of sections. If not rendering bugs and issues will show up because the nanite render must use the LOD 0 sections.
+		if (IsHiresMeshDescriptionValid())
 		{
-			MeshBoundsBuilder += NaniteBounds;
-			bNaniteDataBuilt = true;
+			FBoxSphereBounds NaniteBounds;
+			bool bBuildSuccess = BuildNanite(
+				StaticMesh,
+				StaticMesh->GetHiResSourceModel(),
+				StaticMeshRenderData.LODResources,
+				StaticMeshRenderData.LODVertexFactories,
+				NaniteResources,
+				NaniteSettings,
+				TArrayView< float >(),
+				NaniteBounds);
+
+			if (bBuildSuccess)
+			{
+				MeshBoundsBuilder += NaniteBounds;
+				bNaniteDataBuilt = true;
+			}
 		}
 	}
 
