Downgrade a log verbosity from Error to Display to avoid breaking the CIS.

#rnx

diff --git a/Engine/Source/Developer/MeshBuilder/Private/StaticMeshBuilder.cpp b/Engine/Source/Developer/MeshBuilder/Private/StaticMeshBuilder.cpp
--- a/Engine/Source/Developer/MeshBuilder/Private/StaticMeshBuilder.cpp
+++ b/Engine/Source/Developer/MeshBuilder/Private/StaticMeshBuilder.cpp
@@ -391,7 +391,7 @@ bool FStaticMeshBuilder::Build(FStaticMeshRenderData& StaticMeshRenderData, USta
 						//Validate the number of sections
 						if (HiResMeshDescription->PolygonGroups().Num() > BaseLodMeshDescription->PolygonGroups().Num())
 						{
-							UE_LOG(LogStaticMeshBuilder, Error, TEXT("Invalid hi-res mesh description during Nanite build [%s]. The number of sections from the hires mesh is higher than LOD 0 section count. This is not supported and LOD 0 will be used as a fallback to build nanite data."), *StaticMesh->GetFullName());
+							UE_LOG(LogStaticMeshBuilder, Display, TEXT("Invalid hi-res mesh description during Nanite build [%s]. The number of sections from the hires mesh is higher than LOD 0 section count. This is not supported and LOD 0 will be used as a fallback to build nanite data."), *StaticMesh->GetFullName());
 							bIsValid = false;
 						}
 						else
