Modified logic to allow a static mesh to have an empty mesh description

#rnx

diff --git a/Engine/Source/Developer/MeshBuilder/Private/StaticMeshBuilder.cpp b/Engine/Source/Developer/MeshBuilder/Private/StaticMeshBuilder.cpp
--- a/Engine/Source/Developer/MeshBuilder/Private/StaticMeshBuilder.cpp
+++ b/Engine/Source/Developer/MeshBuilder/Private/StaticMeshBuilder.cpp
@@ -145,6 +145,12 @@ static bool BuildNanite(
 		return false;
 	}
 
+	if (MeshDescription.IsEmpty())
+	{
+		UE_LOG(LogStaticMeshBuilder, Error, TEXT("Cannot build an empty mesh description during Nanite build [%s]."), *StaticMesh->GetFullName());
+		return false;
+	}
+
 	FMeshBuildSettings& BuildSettings = StaticMesh->GetSourceModel(0).BuildSettings;
 	FStaticMeshLODResources& StaticMeshLOD = LODResources[0];
 
@@ -388,19 +394,27 @@ bool FStaticMeshBuilder::Build(FStaticMeshRenderData& StaticMeshRenderData, USta
 				{
 					if (const FMeshDescription* HiResMeshDescription = StaticMesh->GetHiResSourceModel().GetOrCacheMeshDescription())
 					{
-						//Validate the number of sections
-						if (HiResMeshDescription->PolygonGroups().Num() > BaseLodMeshDescription->PolygonGroups().Num())
+						if (HiResMeshDescription->IsEmpty())
 						{
-							UE_LOG(LogStaticMeshBuilder, Display, TEXT("Invalid hi-res mesh description during Nanite build [%s]. The number of sections from the hires mesh is higher than LOD 0 section count. This is not supported and LOD 0 will be used as a fallback to build nanite data."), *StaticMesh->GetFullName());
+							UE_LOG(LogStaticMeshBuilder, Display, TEXT("Invalid hi-res mesh description during Nanite build [%s]. The hires mesh is empty. This is not supported and LOD 0 will be used as a fallback to build nanite data."), *StaticMesh->GetFullName());
 							bIsValid = false;
 						}
 						else
 						{
-							if (HiResMeshDescription->PolygonGroups().Num() < BaseLodMeshDescription->PolygonGroups().Num())
+							//Validate the number of sections
+							if (HiResMeshDescription->PolygonGroups().Num() > BaseLodMeshDescription->PolygonGroups().Num())
+							{
+								UE_LOG(LogStaticMeshBuilder, Display, TEXT("Invalid hi-res mesh description during Nanite build [%s]. The number of sections from the hires mesh is higher than LOD 0 section count. This is not supported and LOD 0 will be used as a fallback to build nanite data."), *StaticMesh->GetFullName());
+								bIsValid = false;
+							}
+							else
 							{
-								UE_LOG(LogStaticMeshBuilder, Display, TEXT("Nanite hi-res mesh description for [%s] has fewer sections than lod 0. Verify you have the proper material id result when nanite is turned on."), *StaticMesh->GetFullName());
+								if (HiResMeshDescription->PolygonGroups().Num() < BaseLodMeshDescription->PolygonGroups().Num())
+								{
+									UE_LOG(LogStaticMeshBuilder, Display, TEXT("Nanite hi-res mesh description for [%s] has fewer sections than lod 0. Verify you have the proper material id result when nanite is turned on."), *StaticMesh->GetFullName());
+								}
+								bIsValid = true;
 							}
-							bIsValid = true;
 						}
 					}
 				}
@@ -490,6 +504,7 @@ bool FStaticMeshBuilder::Build(FStaticMeshRenderData& StaticMeshRenderData, USta
 		float MaxDeviation = 0.0f;
 		FMeshBuildSettings& LODBuildSettings = SrcModel.BuildSettings;
 		bool bIsMeshDescriptionValid = StaticMesh->CloneMeshDescription(LodIndex, MeshDescriptions[LodIndex]);
+		bIsMeshDescriptionValid &= !MeshDescriptions[LodIndex].IsEmpty();
 		FMeshDescriptionHelper MeshDescriptionHelper(&LODBuildSettings);
 
 		FMeshReductionSettings ReductionSettings = LODGroup.GetSettings(SrcModel.ReductionSettings, LodIndex);
@@ -777,6 +792,12 @@ bool FStaticMeshBuilder::BuildMeshVertexPositions(
 	const bool bIsMeshDescriptionValid = SourceModel.CloneMeshDescription(MeshDescription);
 	check(bIsMeshDescriptionValid);
 
+	if (MeshDescription.IsEmpty())
+	{
+		UE_LOG(LogStaticMeshBuilder, Error, TEXT("Cannot build the asset from an empty mesh description."));
+		return false;
+	}
+
 	FMeshBuildSettings& BuildSettings = StaticMesh->GetSourceModel(0).BuildSettings;
 
 	FMeshDescriptionHelper MeshDescriptionHelper(&BuildSettings);
