[Nanite-Skinning] Fix up bone indices when there are bone maps that perform bone index remapping

diff --git a/Engine/Source/Developer/MeshBuilder/Private/SkeletalMeshBuilder.cpp b/Engine/Source/Developer/MeshBuilder/Private/SkeletalMeshBuilder.cpp
--- a/Engine/Source/Developer/MeshBuilder/Private/SkeletalMeshBuilder.cpp
+++ b/Engine/Source/Developer/MeshBuilder/Private/SkeletalMeshBuilder.cpp
@@ -140,7 +140,7 @@ static bool BuildNanite(
 
 			for (int32 Influence = 0; Influence < Section.MaxBoneInfluences; ++Influence)
 			{
-				InputMeshData.Vertices.BoneIndices[Influence][Section.BaseVertexIndex + VertIndex] = SoftVertex.InfluenceBones[Influence];
+				InputMeshData.Vertices.BoneIndices[Influence][Section.BaseVertexIndex + VertIndex] = Section.BoneMap[SoftVertex.InfluenceBones[Influence]];
 				InputMeshData.Vertices.BoneWeights[Influence][Section.BaseVertexIndex + VertIndex] = SoftVertex.InfluenceWeights[Influence];
 			}
 
