#UE Fix NaN bounds on meshes that have no verts that contribute to bounds

diff --git a/Engine/Source/Developer/MeshBuilder/Private/StaticMeshBuilder.cpp b/Engine/Source/Developer/MeshBuilder/Private/StaticMeshBuilder.cpp
--- a/Engine/Source/Developer/MeshBuilder/Private/StaticMeshBuilder.cpp
+++ b/Engine/Source/Developer/MeshBuilder/Private/StaticMeshBuilder.cpp
@@ -932,6 +932,7 @@ void BuildVertexBuffer(
 	const bool bCacheOptimize = (NumVertexInstances < 100000 * 3);
 
 	FBounds3f Bounds;
+	bool bBoundsSet = false;
 
 	FStaticMeshConstAttributes Attributes(MeshDescription);
 
@@ -1134,6 +1135,7 @@ void BuildVertexBuffer(
 				// We are already processing all vertices, so we may as well compute the bounding box here
 				// instead of yet another loop over the vertices at a later point.
 				Bounds += PendingVertex.Position;
+				bBoundsSet = true;
 			}
 
 				RemapVerts[WedgeIndex] = Index;
@@ -1147,6 +1149,12 @@ void BuildVertexBuffer(
 		}
 	}
 
+	if (!bBoundsSet)
+	{
+		// There were no verts that contribute to bounds, so we'll just set a bounds of 0,0,0 to avoid calculating NaNs for Origin, BoxExtent, and SphereRadius below
+		Bounds = FVector3f(0.f, 0.f, 0.f);
+	}
+
 	// Calculate the bounding sphere, using the center of the bounding box as the origin.
 	FVector3f Center = Bounds.GetCenter();
 	float RadiusSqr = 0.0f;
