Provide TargetPlatform to FStaticMeshBuilder.
- used to decide whether to build nanite data and soon ray tracing proxies.

diff --git a/Engine/Source/Developer/MeshBuilder/Private/MeshBuilderModule.cpp b/Engine/Source/Developer/MeshBuilder/Private/MeshBuilderModule.cpp
--- a/Engine/Source/Developer/MeshBuilder/Private/MeshBuilderModule.cpp
+++ b/Engine/Source/Developer/MeshBuilder/Private/MeshBuilderModule.cpp
@@ -26,6 +26,8 @@ public:
 		// Unregister any modular features here
 	}
 
+	virtual bool BuildMesh(class FStaticMeshRenderData& OutRenderData, const FStaticMeshBuildParameters& BuildParameters) override;
+
 	virtual bool BuildMesh(FStaticMeshRenderData& OutRenderData, UObject* Mesh, const FStaticMeshLODGroup& LODGroup, bool bAllowNanite) override;
 
 	virtual bool BuildMeshVertexPositions(
@@ -42,13 +44,18 @@ private:
 
 IMPLEMENT_MODULE(FMeshBuilderModule, MeshBuilder );
 
+bool FMeshBuilderModule::BuildMesh(FStaticMeshRenderData& OutRenderData, const FStaticMeshBuildParameters& BuildParameters)
+{
+	//Call the static mesh builder
+	return FStaticMeshBuilder().Build(OutRenderData, BuildParameters);
+}
+
 bool FMeshBuilderModule::BuildMesh(FStaticMeshRenderData& OutRenderData, class UObject* Mesh, const FStaticMeshLODGroup& LODGroup, bool bAllowNanite)
 {
 	UStaticMesh* StaticMesh = Cast<UStaticMesh>(Mesh);
 	if (StaticMesh != nullptr)
 	{
-		//Call the static mesh builder
-		return FStaticMeshBuilder().Build(OutRenderData, StaticMesh, LODGroup, bAllowNanite);
+		return BuildMesh(OutRenderData, FStaticMeshBuildParameters(StaticMesh, nullptr, LODGroup));
 	}
 	return false;
 }
diff --git a/Engine/Source/Developer/MeshBuilder/Private/StaticMeshBuilder.cpp b/Engine/Source/Developer/MeshBuilder/Private/StaticMeshBuilder.cpp
--- a/Engine/Source/Developer/MeshBuilder/Private/StaticMeshBuilder.cpp
+++ b/Engine/Source/Developer/MeshBuilder/Private/StaticMeshBuilder.cpp
@@ -333,8 +333,19 @@ static bool BuildNanite(
 }
 
 
-bool FStaticMeshBuilder::Build(FStaticMeshRenderData& StaticMeshRenderData, UStaticMesh* StaticMesh, const FStaticMeshLODGroup& LODGroup, bool bTargetSupportsNanite)
+bool FStaticMeshBuilder::Build(FStaticMeshRenderData& StaticMeshRenderData, const FStaticMeshBuildParameters& BuildParameters)
 {
+	if (BuildParameters.TargetPlatform == nullptr)
+	{
+		UE_LOG(LogStaticMeshBuilder, Error, TEXT("Provided FStaticMeshBuildParameters must have a valid TargetPlatform."));
+		return false;
+	}
+
+	UStaticMesh* StaticMesh = BuildParameters.StaticMesh;
+	const FStaticMeshLODGroup& LODGroup = BuildParameters.LODGroup;
+
+	const bool bTargetSupportsNanite = DoesTargetPlatformSupportNanite(BuildParameters.TargetPlatform);
+
 	const bool bNaniteBuildEnabled = StaticMesh->IsNaniteEnabled();
 	const bool bHaveHiResSourceModel = StaticMesh->IsHiResMeshDescriptionValid();
 	int32 NumTasks = (bNaniteBuildEnabled && bHaveHiResSourceModel) ? (StaticMesh->GetNumSourceModels() + 1) : (StaticMesh->GetNumSourceModels());
@@ -772,6 +783,15 @@ bool FStaticMeshBuilder::Build(FStaticMeshRenderData& StaticMeshRenderData, USta
 	return true;
 }
 
+bool FStaticMeshBuilder::Build(
+	FStaticMeshRenderData& OutRenderData,
+	UStaticMesh* StaticMesh,
+	const FStaticMeshLODGroup& LODGroup,
+	bool bAllowNanite)
+{
+	return Build(OutRenderData, FStaticMeshBuildParameters(StaticMesh, nullptr, LODGroup));
+}
+
 bool FStaticMeshBuilder::BuildMeshVertexPositions(
 	UStaticMesh* StaticMesh,
 	TArray<uint32>& BuiltIndices,
diff --git a/Engine/Source/Developer/MeshBuilder/Public/IMeshBuilderModule.h b/Engine/Source/Developer/MeshBuilder/Public/IMeshBuilderModule.h
--- a/Engine/Source/Developer/MeshBuilder/Public/IMeshBuilderModule.h
+++ b/Engine/Source/Developer/MeshBuilder/Public/IMeshBuilderModule.h
@@ -8,7 +8,6 @@
 #include "Interfaces/ITargetPlatform.h"
 #include "Interfaces/ITargetPlatformManagerModule.h"
 
-struct FStaticMeshBuildVertex;
 struct FStaticMeshSection;
 class FStaticMeshSectionArray;
 
@@ -29,6 +28,9 @@ public:
 
 	virtual void AppendToDDCKey(FString& DDCKey, bool bSkeletal) { }
 
+	virtual bool BuildMesh(class FStaticMeshRenderData& OutRenderData, const struct FStaticMeshBuildParameters& BuildParameters) = 0;
+
+	UE_DEPRECATED(5.5, "Use FStaticMeshBuildParameters instead.")
 	virtual bool BuildMesh(class FStaticMeshRenderData& OutRenderData, class UObject* Mesh, const class FStaticMeshLODGroup& LODGroup, bool bAllowNanite) = 0;
 
 	virtual bool BuildMeshVertexPositions(class UObject* StaticMesh, TArray<uint32>& Indices, TArray<FVector3f>& Vertices, FStaticMeshSectionArray& Sections) = 0;
diff --git a/Engine/Source/Developer/MeshBuilder/Public/MeshBuilder.h b/Engine/Source/Developer/MeshBuilder/Public/MeshBuilder.h
--- a/Engine/Source/Developer/MeshBuilder/Public/MeshBuilder.h
+++ b/Engine/Source/Developer/MeshBuilder/Public/MeshBuilder.h
@@ -9,11 +9,9 @@
 class FStaticMeshLODGroup;
 class FStaticMeshRenderData;
 class FStaticMeshSectionArray;
-class USkeletalMesh;
 class UStaticMesh;
+struct FStaticMeshBuildParameters;
 struct FSkeletalMeshBuildParameters;
-struct FStaticMeshBuildVertex;
-struct FStaticMeshSection;
 struct FMeshBuildVertexData;
 
 /**
@@ -28,6 +26,9 @@ public:
 	/**
 	 * Build function should be override and is the starting point for static mesh builders
 	 */
+	virtual bool Build(FStaticMeshRenderData& OutRenderData, const FStaticMeshBuildParameters& BuildParameters) = 0;
+
+	UE_DEPRECATED(5.5, "Use FStaticMeshBuildParameters instead.")
 	virtual bool Build(
 		FStaticMeshRenderData& OutRenderData,
 		UStaticMesh* StaticMesh,
diff --git a/Engine/Source/Developer/MeshBuilder/Public/SkeletalMeshBuilder.h b/Engine/Source/Developer/MeshBuilder/Public/SkeletalMeshBuilder.h
--- a/Engine/Source/Developer/MeshBuilder/Public/SkeletalMeshBuilder.h
+++ b/Engine/Source/Developer/MeshBuilder/Public/SkeletalMeshBuilder.h
@@ -18,6 +18,14 @@ public:
 	FSkeletalMeshBuilder();
 
 	//No support for static mesh build in this class
+	virtual bool Build(FStaticMeshRenderData& OutRenderData, const struct FStaticMeshBuildParameters& BuildParameters) override
+	{
+		bool No_Support_For_StaticMesh_Build_In_FSkeletalMeshBuilder_Class = false;
+		check(No_Support_For_StaticMesh_Build_In_FSkeletalMeshBuilder_Class);
+		return false;
+	}
+
+	UE_DEPRECATED(5.5, "Use FStaticMeshBuildParameters instead.")
 	virtual bool Build(FStaticMeshRenderData& OutRenderData, UStaticMesh* StaticMesh, const FStaticMeshLODGroup& LODGroup, bool bAllowNanite) override
 	{
 		bool No_Support_For_StaticMesh_Build_In_FSkeletalMeshBuilder_Class = false;
diff --git a/Engine/Source/Developer/MeshBuilder/Public/StaticMeshBuilder.h b/Engine/Source/Developer/MeshBuilder/Public/StaticMeshBuilder.h
--- a/Engine/Source/Developer/MeshBuilder/Public/StaticMeshBuilder.h
+++ b/Engine/Source/Developer/MeshBuilder/Public/StaticMeshBuilder.h
@@ -7,21 +7,23 @@
 
 DECLARE_LOG_CATEGORY_EXTERN(LogStaticMeshBuilder, Log, All);
 
+class ITargetPlatform;
 class UStaticMesh;
 class FStaticMeshRenderData;
 class FStaticMeshLODGroup;
-class USkeletalMesh;
 struct FOverlappingCorners;
 struct FMeshDescription;
 struct FMeshBuildSettings;
 
-
 class MESHBUILDER_API FStaticMeshBuilder : public FMeshBuilder
 {
 public:
 	FStaticMeshBuilder();
 	virtual ~FStaticMeshBuilder() {}
 
+	virtual bool Build(FStaticMeshRenderData& OutRenderData, const FStaticMeshBuildParameters& BuildParameters) override;
+
+	UE_DEPRECATED(5.5, "Use FStaticMeshBuildParameters instead.")
 	virtual bool Build(
 		FStaticMeshRenderData& OutRenderData,
 		UStaticMesh* StaticMesh,
diff --git a/Engine/Source/Runtime/Engine/Classes/Engine/StaticMesh.h b/Engine/Source/Runtime/Engine/Classes/Engine/StaticMesh.h
--- a/Engine/Source/Runtime/Engine/Classes/Engine/StaticMesh.h
+++ b/Engine/Source/Runtime/Engine/Classes/Engine/StaticMesh.h
@@ -2151,6 +2151,19 @@ public:
 	bool bHasRenderDataChanged = false;
 };
 
+struct FStaticMeshBuildParameters
+{
+	FStaticMeshBuildParameters(UStaticMesh* InStaticMesh, const ITargetPlatform* InTargetPlatform, const FStaticMeshLODGroup& InLODGroup)
+		: StaticMesh(InStaticMesh)
+		, TargetPlatform(InTargetPlatform)
+		, LODGroup(InLODGroup)
+	{}
+
+	UStaticMesh* StaticMesh;
+	const ITargetPlatform* TargetPlatform;
+	const FStaticMeshLODGroup& LODGroup;
+};
+
 namespace UE::Private::StaticMesh
 {
 #if WITH_EDITOR
diff --git a/Engine/Source/Runtime/Engine/Private/StaticMesh.cpp b/Engine/Source/Runtime/Engine/Private/StaticMesh.cpp
--- a/Engine/Source/Runtime/Engine/Private/StaticMesh.cpp
+++ b/Engine/Source/Runtime/Engine/Private/StaticMesh.cpp
@@ -3140,10 +3140,7 @@ void FStaticMeshRenderData::Cache(const ITargetPlatform* TargetPlatform, UStatic
 			{
 				IMeshBuilderModule& MeshBuilderModule = IMeshBuilderModule::GetForPlatform(TargetPlatform);
 
-				// Check if the target platform supports Nanite at all
-				const bool bAllowNanite = DoesTargetPlatformSupportNanite(TargetPlatform);
-
-				if (!MeshBuilderModule.BuildMesh(*this, Owner, LODGroup, bAllowNanite))
+				if (!MeshBuilderModule.BuildMesh(*this, FStaticMeshBuildParameters(Owner, TargetPlatform, LODGroup)))
 				{
 					UE_LOG(LogStaticMesh, Error, TEXT("Failed to build static mesh. See previous line(s) for details."));
 					return;
