Fix linux compile warning, and possible UB.

diff --git a/Engine/Source/Developer/MeshBuilder/Private/StaticMeshBuilder.cpp b/Engine/Source/Developer/MeshBuilder/Private/StaticMeshBuilder.cpp
--- a/Engine/Source/Developer/MeshBuilder/Private/StaticMeshBuilder.cpp
+++ b/Engine/Source/Developer/MeshBuilder/Private/StaticMeshBuilder.cpp
@@ -390,13 +390,13 @@ static void BuildNaniteFallbackMeshDescription(
 
 		for (uint32 TriIndex = 0; TriIndex < Section.NumTriangles; ++TriIndex)
 		{
-			TConstArrayView<FVertexInstanceID> TriVertInstanceIDs = {
+			const FVertexInstanceID TriVertInstanceIDs[] = {
 				FVertexInstanceID(InMeshData.TriangleIndices[Section.FirstIndex + TriIndex * 3 + 0]),
 				FVertexInstanceID(InMeshData.TriangleIndices[Section.FirstIndex + TriIndex * 3 + 1]),
 				FVertexInstanceID(InMeshData.TriangleIndices[Section.FirstIndex + TriIndex * 3 + 2])
 			};
 
-			OutMesh.CreateTriangle(PolygonGroupID, TriVertInstanceIDs);
+			OutMesh.CreateTriangle(PolygonGroupID, MakeConstArrayView(TriVertInstanceIDs, 3));
 		}
 	}
 }
