Fix Mac compile error

diff --git a/Engine/Source/Developer/MeshBuilder/Private/StaticMeshBuilder.cpp b/Engine/Source/Developer/MeshBuilder/Private/StaticMeshBuilder.cpp
--- a/Engine/Source/Developer/MeshBuilder/Private/StaticMeshBuilder.cpp
+++ b/Engine/Source/Developer/MeshBuilder/Private/StaticMeshBuilder.cpp
@@ -299,7 +299,8 @@ static bool InitNaniteBuildInput(
 		OutData.Sections[SectionIndex].MaterialIndex = Context.StaticMesh->GetSectionInfoMap().Get(0, SectionIndex).MaterialIndex;
 	}
 
-	OutData.VertexBounds = FBounds3f(FVector3f(OutBounds.Origin - OutBounds.BoxExtent), FVector3f(OutBounds.Origin + OutBounds.BoxExtent));
+	OutData.VertexBounds.Min = FVector4f(FVector3f(OutBounds.Origin - OutBounds.BoxExtent), 0.0f);
+	OutData.VertexBounds.Max = FVector4f(FVector3f(OutBounds.Origin + OutBounds.BoxExtent), 0.0f);
 
 	TVertexInstanceAttributesRef<FVector2f const> VertexInstanceUVs = MeshDescription.VertexInstanceAttributes().GetAttributesRef<FVector2f>(MeshAttribute::VertexInstance::TextureCoordinate);
 	OutData.NumTexCoords = VertexInstanceUVs.IsValid() ? VertexInstanceUVs.GetNumChannels() : 0;
