Fix a crash in recent static mesh builder changes that were encountered with some assets in the Lyra project.

#tests Lyra PIE automation test

diff --git a/Engine/Source/Developer/MeshBuilder/Private/StaticMeshBuilder.cpp b/Engine/Source/Developer/MeshBuilder/Private/StaticMeshBuilder.cpp
--- a/Engine/Source/Developer/MeshBuilder/Private/StaticMeshBuilder.cpp
+++ b/Engine/Source/Developer/MeshBuilder/Private/StaticMeshBuilder.cpp
@@ -387,7 +387,9 @@ static void BuildNaniteFallbackMeshDescription(
 	for (const FStaticMeshSection& Section : InMeshData.Sections)
 	{
 		const FPolygonGroupID PolygonGroupID = OutMesh.CreatePolygonGroup();
-		PolygonGroupMaterialSlotNames.Set(PolygonGroupID, StaticMaterials[Section.MaterialIndex].ImportedMaterialSlotName);
+		const FName MaterialSlotName = StaticMaterials.IsValidIndex(Section.MaterialIndex) ?
+			StaticMaterials[Section.MaterialIndex].ImportedMaterialSlotName : NAME_None;
+		PolygonGroupMaterialSlotNames.Set(PolygonGroupID, MaterialSlotName);
 
 		for (uint32 TriIndex = 0; TriIndex < Section.NumTriangles; ++TriIndex)
 		{
