Fixed NaniteBuildError localization conflict

diff --git a/Engine/Source/Developer/MeshBuilder/Private/StaticMeshBuilder.cpp b/Engine/Source/Developer/MeshBuilder/Private/StaticMeshBuilder.cpp
--- a/Engine/Source/Developer/MeshBuilder/Private/StaticMeshBuilder.cpp
+++ b/Engine/Source/Developer/MeshBuilder/Private/StaticMeshBuilder.cpp
@@ -319,7 +319,7 @@ static bool InitNaniteBuildInput(FStaticMeshNaniteBuildContext& Context, Nanite:
 
 	if (!Context.Builder->BuildMaterialIndices(OutData.Sections, TriangleCount, OutData.MaterialIndices))
 	{
-		UE_LOGFMT_NSLOC(LogStaticMesh, Warning, "StaticMesh", "NaniteBuildError", "Failed to build Nanite from static mesh. See previous line(s) for details.");
+		UE_LOGFMT_NSLOC(LogStaticMesh, Warning, "StaticMesh", "NaniteBuildStaticMeshError", "Failed to build Nanite from static mesh. See previous line(s) for details.");
 		return false;
 	}
 
@@ -357,7 +357,7 @@ static bool InitNaniteAssemblyData(
 #if NANITE_ASSEMBLY_DATA
 	if (!Context.StaticMesh->HasCachedNaniteAssemblyReferences())
 	{
-		UE_LOGFMT_NSLOC(LogStaticMesh, Warning, "StaticMesh", "NaniteBuildError",
+		UE_LOGFMT_NSLOC(LogStaticMesh, Warning, "StaticMesh", "NaniteBuildAssemblyUncachedError",
 			"Failed to build Nanite Assembly static mesh {MeshPath}. The referenced static meshes were not cached before build.",
 			("MeshPath", Context.StaticMesh->GetPathName()));
 		return false;
@@ -390,7 +390,7 @@ static bool InitNaniteAssemblyData(
 
 		if (!Resource.IsValid())
 		{
-			UE_LOGFMT_NSLOC(LogStaticMesh, Warning, "StaticMesh", "NaniteBuildError",
+			UE_LOGFMT_NSLOC(LogStaticMesh, Warning, "StaticMesh", "NaniteBuildAssemblyPartError",
 				"Failed to build Nanite assembly part from static mesh ({PartPath}). See previous line(s) for details.",
 				("PartPath", AssemblyData.Parts[PartIndex].MeshObjectPath.GetAssetName()));
 			return false;
