[HWRT] Call InitializeRayTracingRepresentationFromRenderingLODs() if RayTracingProxy was not created during FStaticMeshBuilder::Build(...).

diff --git a/Engine/Source/Developer/MeshBuilder/Private/StaticMeshBuilder.cpp b/Engine/Source/Developer/MeshBuilder/Private/StaticMeshBuilder.cpp
--- a/Engine/Source/Developer/MeshBuilder/Private/StaticMeshBuilder.cpp
+++ b/Engine/Source/Developer/MeshBuilder/Private/StaticMeshBuilder.cpp
@@ -979,15 +979,13 @@ bool FStaticMeshBuilder::Build(FStaticMeshRenderData& StaticMeshRenderData, cons
 
 	if (StaticMesh->bSupportRayTracing && BuildParameters.TargetPlatform->UsesRayTracing())
 	{
-		const bool bUsingRenderingLODs = true;
-
-		if (bUsingRenderingLODs)
+		if (!StaticMeshRenderData.RayTracingProxy)
 		{
 			StaticMeshRenderData.InitializeRayTracingRepresentationFromRenderingLODs();
 		}
 		else
 		{
-			unimplemented();
+			check(!StaticMeshRenderData.RayTracingProxy->bUsingRenderingLODs);
 		}
 	}
 	
