Fix a bug in static mesh building where LODs built as a reduction of LOD0 from the Nanite fallback will apply build scale to the mesh twice.

diff --git a/Engine/Source/Developer/MeshBuilder/Private/StaticMeshBuilder.cpp b/Engine/Source/Developer/MeshBuilder/Private/StaticMeshBuilder.cpp
--- a/Engine/Source/Developer/MeshBuilder/Private/StaticMeshBuilder.cpp
+++ b/Engine/Source/Developer/MeshBuilder/Private/StaticMeshBuilder.cpp
@@ -121,6 +121,43 @@ static void CorrectFallbackSettings( FMeshNaniteSettings& NaniteSettings, int32
 	}
 }
 
+static void ScaleStaticMeshVertex(
+	FVector3f& Position,
+	FVector3f& TangentX,
+	FVector3f& TangentY,
+	FVector3f& TangentZ,
+	FVector3f Scale,
+	bool bNeedTangents,
+	bool bUseLegacyTangentScaling
+)
+{
+	Position *= Scale;
+	if (bNeedTangents)
+	{
+		if (bUseLegacyTangentScaling)
+		{
+			// Apply incorrect inverse scale to tangents to match an old bug, for legacy assets only
+			TangentX /= Scale;
+			TangentY /= Scale;
+		}
+		else
+		{
+			// Tangents should transform by directly applying the same scale as the geometry; it's only the normal that needs an inverse scale
+			TangentX *= Scale;
+			TangentY *= Scale;
+		}
+		TangentX.Normalize();
+		TangentY.Normalize();
+	}
+	else
+	{
+		TangentX = FVector3f(1.0f, 0.0f, 0.0f);
+		TangentY = FVector3f(0.0f, 1.0f, 0.0f);
+	}
+	TangentZ /= Scale;
+	TangentZ.Normalize();
+}
+
 struct FStaticMeshNaniteBuildContext
 {
 	FMeshNaniteSettings Settings;
@@ -435,6 +472,11 @@ static void BuildNaniteFallbackMeshDescription(
 )
 {
 	OutMesh.Empty();
+	
+	// Lod zero was built with scaling build settings, we have to remove the scaling from the data since the other LODs build will also apply the scaling.
+	const FVector3f InverseBuildScale = FVector3f(FVector(1.0) / Context.SourceModel->BuildSettings.BuildScale3D);
+	const bool bBuildScaleActive = !InverseBuildScale.Equals(FVector3f(1.0f), UE_SMALL_NUMBER);
+	const bool bUseLegacyTangentScaling = Context.StaticMesh->GetLegacyTangentScaling();
 
 	FStaticMeshAttributes Attributes(OutMesh);
 	Attributes.Register();
@@ -473,15 +515,29 @@ static void BuildNaniteFallbackMeshDescription(
 		OutMesh.CreateVertexWithID(VertexID);
 		OutMesh.CreateVertexInstanceWithID(VertexInstanceID, VertexID);
 
-		const FVector3f Position = InMeshData.Vertices.Position[InVertIndex];
-		const FVector3f TangentX = InMeshData.Vertices.TangentX[InVertIndex];
-		const FVector3f TangentY = InMeshData.Vertices.TangentY[InVertIndex];
-		const FVector3f TangentZ = InMeshData.Vertices.TangentZ[InVertIndex];
+		FVector3f Position = InMeshData.Vertices.Position[InVertIndex];
+		FVector3f TangentX = InMeshData.Vertices.TangentX[InVertIndex];
+		FVector3f TangentY = InMeshData.Vertices.TangentY[InVertIndex];
+		FVector3f TangentZ = InMeshData.Vertices.TangentZ[InVertIndex];
+
+		if (bBuildScaleActive)
+		{
+			ScaleStaticMeshVertex(
+				Position,
+				TangentX,
+				TangentY,
+				TangentZ,
+				InverseBuildScale,
+				true, // bNeedTangents
+				bUseLegacyTangentScaling
+			);
+		}
+
 		const float BinormalSign = GetBasisDeterminantSign(FVector(TangentX), FVector(TangentY), FVector(TangentZ));
 		const FColor Color = InMeshData.Vertices.Color.IsValidIndex(InVertIndex) ?
 			InMeshData.Vertices.Color[InVertIndex] : FColor::White;
 
-		VertexPositions.Set(VertexID, InMeshData.Vertices.Position[InVertIndex]);
+		VertexPositions.Set(VertexID, Position);
 		VertexInstanceNormals.Set(VertexInstanceID, TangentZ);
 		VertexInstanceTangents.Set(VertexInstanceID, TangentX);
 		VertexInstanceBinormalSigns.Set(VertexInstanceID, BinormalSign);
@@ -1517,31 +1573,20 @@ void BuildVertexBuffer(
 
 			FPendingVertex PendingVertex;
 
-			PendingVertex.Position = VertexPosition * BuildScale;
-			if (bNeedTangents)
-			{
-				if (bUseLegacyTangentScaling)
-				{
-					// Apply incorrect inverse scale to tangents to match an old bug, for legacy assets only
-					PendingVertex.TangentX = VertexInstanceTangent / BuildScale;
-					PendingVertex.TangentY = ((VertexInstanceNormal ^ VertexInstanceTangent) * VertexInstanceBinormalSign) / BuildScale;
-				}
-				else
-				{
-					// Tangents should transform by directly applying the same scale as the geometry; it's only the normal that needs an inverse scale
-					PendingVertex.TangentX = VertexInstanceTangent * BuildScale;
-					PendingVertex.TangentY = ((VertexInstanceNormal ^ VertexInstanceTangent) * VertexInstanceBinormalSign) * BuildScale;
-				}
-				PendingVertex.TangentX.Normalize();
-				PendingVertex.TangentY.Normalize();
-			}
-			else
-			{
-				PendingVertex.TangentX = FVector3f(1.0f, 0.0f, 0.0f);
-				PendingVertex.TangentY = FVector3f(0.0f, 1.0f, 0.0f);
-			}
-			PendingVertex.TangentZ = VertexInstanceNormal / BuildScale;
-			PendingVertex.TangentZ.Normalize();
+			PendingVertex.Position = VertexPosition;
+			PendingVertex.TangentX = VertexInstanceTangent;
+			PendingVertex.TangentY = (VertexInstanceNormal ^ VertexInstanceTangent) * VertexInstanceBinormalSign;
+			PendingVertex.TangentZ = VertexInstanceNormal;
+			
+			ScaleStaticMeshVertex(
+				PendingVertex.Position,
+				PendingVertex.TangentX,
+				PendingVertex.TangentY,
+				PendingVertex.TangentZ,
+				BuildScale,
+				bNeedTangents,
+				bUseLegacyTangentScaling
+			);
 
 			FColor VertexColor = FColor::White;
 			if (bHasColors)
