Skeletal mesh builder refactor in preparation of upcoming Nanite Assembly changes.

#tests Lyra & CitySample PIE
#rnx

diff --git a/Engine/Source/Developer/MeshBuilder/Private/SkeletalMeshBuilder.cpp b/Engine/Source/Developer/MeshBuilder/Private/SkeletalMeshBuilder.cpp
--- a/Engine/Source/Developer/MeshBuilder/Private/SkeletalMeshBuilder.cpp
+++ b/Engine/Source/Developer/MeshBuilder/Private/SkeletalMeshBuilder.cpp
@@ -29,224 +29,80 @@
 
 DEFINE_LOG_CATEGORY(LogSkeletalMeshBuilder);
 
-namespace SkeletalMeshBuilderOptimization
+namespace SkeletalMeshBuilderPrivate
 {
-	void CacheOptimizeIndexBuffer(TArray<uint16>& Indices)
-	{
-		BuildOptimizationThirdParty::CacheOptimizeIndexBuffer(Indices);
-	}
-
-	void CacheOptimizeIndexBuffer(TArray<uint32>& Indices)
-	{
-		BuildOptimizationThirdParty::CacheOptimizeIndexBuffer(Indices);
-	}
-}
-
-struct InfluenceMap
-{
-	FORCEINLINE bool operator()(const float& A, const float& B) const
-	{
-		return B > A;
-	}
-};
-
-struct FSkeletalMeshVertInstanceIDAndZ
-{
-	FVertexInstanceID Index;
-	float Z;
-};
-
-FSkeletalMeshBuilder::FSkeletalMeshBuilder()
-{
-}
 
-static bool BuildNanite(
-	USkeletalMesh* SkeletalMesh,
-	FSkeletalMeshLODModel& LODModel,
-	const FMeshDescription& MeshDescription,
-	Nanite::FResources& NaniteResources
-)
+/** Context data for a skeletal mesh build */
+class FContext
 {
-	TRACE_CPUPROFILER_EVENT_SCOPE(FStaticMeshBuilder::BuildNanite);
-
-	Nanite::IBuilderModule& NaniteBuilderModule = Nanite::IBuilderModule::Get();
-
-	// Build new vertex buffers
-	Nanite::IBuilderModule::FInputMeshData InputMeshData;
-
-	InputMeshData.NumTexCoords = LODModel.NumTexCoords;
-
-	InputMeshData.MaterialIndices.SetNumUninitialized(LODModel.IndexBuffer.Num() / 3);
-
-	InputMeshData.Vertices.Position.SetNumUninitialized(LODModel.NumVertices);
-	InputMeshData.Vertices.TangentX.SetNumUninitialized(LODModel.NumVertices);
-	InputMeshData.Vertices.TangentY.SetNumUninitialized(LODModel.NumVertices);
-	InputMeshData.Vertices.TangentZ.SetNumUninitialized(LODModel.NumVertices);
-
-	InputMeshData.Vertices.UVs.SetNum(LODModel.NumTexCoords);
-	for (uint32 UVCoord = 0; UVCoord < LODModel.NumTexCoords; ++UVCoord)
+public:
+	USkeletalMesh*						SkeletalMesh		= nullptr;
+	IMeshUtilities*						MeshUtilities		= nullptr;
+	Nanite::IBuilderModule*				NaniteBuilder		= nullptr;
+	const FSkeletalMeshLODInfo*			LODInfo				= nullptr;
+	const FMeshDescription* 			MeshDescription		= nullptr;
+	const ITargetPlatform*				TargetPlatform		= nullptr;
+	FSkeletalMeshImportData 			ImportData;
+	IMeshUtilities::MeshBuildOptions	Options;
+	FMeshNaniteSettings 				NaniteSettings;
+	int32 								LODIndex			= INDEX_NONE;
+	bool								bRegenDepLODs		= false;
+
+	FContext(const FSkeletalMeshBuildParameters& BuildParameters)
+		: SkeletalMesh(BuildParameters.SkeletalMesh)
+		, MeshUtilities(&FModuleManager::Get().LoadModuleChecked<IMeshUtilities>("MeshUtilities"))
+		, TargetPlatform(BuildParameters.TargetPlatform)
+		, LODIndex(BuildParameters.LODIndex)
+		, bRegenDepLODs(BuildParameters.bRegenDepLODs)
 	{
-		InputMeshData.Vertices.UVs[UVCoord].SetNumUninitialized(LODModel.NumVertices);
-	}
+		if (SkeletalMesh->IsNaniteEnabled())
+		{
+			NaniteBuilder = &Nanite::IBuilderModule::Get();
+		}
 
-	// We can save memory by figuring out the max number of influences across all sections instead of allocating MAX_TOTAL_INFLUENCES
-	// Also check if any of the sections actually require 16bit, or if 8bit will suffice
-	bool b16BitSkinning = false;
-	InputMeshData.NumBoneInfluences = 0;
-	for (const FSkelMeshSection& Section : LODModel.Sections)
-	{
-		InputMeshData.NumBoneInfluences = FMath::Max< uint8 >(InputMeshData.NumBoneInfluences, uint32(Section.MaxBoneInfluences));
-		b16BitSkinning |= Section.Use16BitBoneIndex();
+		Init();
 	}
 
-	InputMeshData.Vertices.BoneIndices.SetNum(InputMeshData.NumBoneInfluences);
-	InputMeshData.Vertices.BoneWeights.SetNum(InputMeshData.NumBoneInfluences);
-	for (uint32 Influence = 0; Influence < InputMeshData.NumBoneInfluences; ++Influence)
+	bool 					IsNaniteBuildEnabled() const		{ return NaniteBuilder != nullptr; }
+	
+	FSkeletalMeshLODModel&	GetLODModel()
 	{
-		InputMeshData.Vertices.BoneIndices[Influence].SetNumZeroed(LODModel.NumVertices);
-		InputMeshData.Vertices.BoneWeights[Influence].SetNumZeroed(LODModel.NumVertices);
+		// NOTE: We don't cache this because some build steps (like reductions) might cause it to be reallocated
+		return SkeletalMesh->GetImportedModel()->LODModels[LODIndex];
 	}
 
-	// TODO: Nanite-Skinning
-	//InputMeshData.Vertices.Color.SetNumUninitialized(LODModel.NumVertices);
-
-	InputMeshData.TriangleIndices = LODModel.IndexBuffer;
-
-	uint32 CheckIndices = 0;
-	uint32 CheckVertices = 0;
-
-	for (int32 SectionIndex = 0; SectionIndex < LODModel.Sections.Num(); ++SectionIndex)
+	int32 GetNumTexCoords() const
 	{
-		const FSkelMeshSection& Section = LODModel.Sections[SectionIndex];
-
-		check(CheckIndices  == Section.BaseIndex);
-		check(CheckVertices == Section.BaseVertexIndex);
-
-		for (int32 VertIndex = 0; VertIndex < Section.SoftVertices.Num(); ++VertIndex)
-		{
-			const FSoftSkinVertex& SoftVertex = Section.SoftVertices[VertIndex];
-
-			InputMeshData.Vertices.Position[Section.BaseVertexIndex + VertIndex] = SoftVertex.Position;
-			InputMeshData.Vertices.TangentX[Section.BaseVertexIndex + VertIndex] = SoftVertex.TangentX;
-			InputMeshData.Vertices.TangentY[Section.BaseVertexIndex + VertIndex] = SoftVertex.TangentY;
-			InputMeshData.Vertices.TangentZ[Section.BaseVertexIndex + VertIndex] = SoftVertex.TangentZ;
-
-			InputMeshData.VertexBounds += SoftVertex.Position;
-
-			for (uint32 UVCoord = 0; UVCoord < LODModel.NumTexCoords; ++UVCoord)
-			{
-				InputMeshData.Vertices.UVs[UVCoord][Section.BaseVertexIndex + VertIndex] = SoftVertex.UVs[UVCoord];
-			}
-
-			for (int32 Influence = 0; Influence < Section.MaxBoneInfluences; ++Influence)
-			{
-				InputMeshData.Vertices.BoneIndices[Influence][Section.BaseVertexIndex + VertIndex] = Section.BoneMap[SoftVertex.InfluenceBones[Influence]];
-				InputMeshData.Vertices.BoneWeights[Influence][Section.BaseVertexIndex + VertIndex] = SoftVertex.InfluenceWeights[Influence];
-			}
-
-			//InputMeshData.Vertices.Color[Section.BaseVertexIndex + VertIndex] = SoftVertex.Color;
-		}
-
-		for (uint32 MaterialIndex = 0; MaterialIndex < Section.NumTriangles; ++MaterialIndex)
-		{
-			InputMeshData.MaterialIndices[(CheckIndices / 3) + MaterialIndex] = Section.MaterialIndex;
-		}
-
-		CheckIndices += Section.NumTriangles * 3;
-		CheckVertices += Section.NumVertices;
+		//We need to send rendering at least one tex coord buffer
+		return FMath::Max<int32>(1, ImportData.NumTexCoords);
 	}
 
-	check(CheckVertices == LODModel.NumVertices);
-	check(CheckIndices == LODModel.IndexBuffer.Num());
-
-	InputMeshData.TriangleCounts.Add(LODModel.IndexBuffer.Num() / 3);
-
-	FMeshNaniteSettings NaniteSettings = SkeletalMesh->NaniteSettings;
-	NaniteSettings.KeepPercentTriangles = 1.0f;
-	NaniteSettings.TrimRelativeError = 0.0f;
-	NaniteSettings.FallbackPercentTriangles = 1.0f; // 100% - no reduction
-	NaniteSettings.FallbackRelativeError = 0.0f;
-
-	if (!NaniteBuilderModule.Build(
-		NaniteResources,
-		InputMeshData,
-		nullptr, // OutFallbackMeshData
-		nullptr, // OutRayTracingFallbackMeshData
-		nullptr, // RayTracingFallbackBuildSettings
-		NaniteSettings))
+	void UnbindClothingAndBackup()
 	{
-		UE_LOG(LogStaticMesh, Error, TEXT("Failed to build Nanite for skeletal mesh. See previous line(s) for details."));
-		return false;
+		//We want to backup in case the LODModel is regenerated, this data is use to validate in the UI if the ddc must be rebuild
+		BackupBuildStringID = SkeletalMesh->GetImportedModel()->LODModels[LODIndex].BuildStringID;
+		FLODUtilities::UnbindClothingAndBackup(SkeletalMesh, ClothingBindings, LODIndex);
 	}
 
-	return true;
-}
-
-bool FSkeletalMeshBuilder::Build(const FSkeletalMeshBuildParameters& SkeletalMeshBuildParameters)
-{
-	TRACE_CPUPROFILER_EVENT_SCOPE(FSkeletalMeshBuilder::Build);
-	const int32 LODIndex = SkeletalMeshBuildParameters.LODIndex;
-	USkeletalMesh* SkeletalMesh = SkeletalMeshBuildParameters.SkeletalMesh;
-
-	check(SkeletalMesh->GetImportedModel());
-	check(SkeletalMesh->GetImportedModel()->LODModels.IsValidIndex(LODIndex));
-	check(SkeletalMesh->GetLODInfo(LODIndex) != nullptr);
-	
-	const FSkeletalMeshLODInfo* LODInfo = SkeletalMesh->GetLODInfo(LODIndex);
-	//We want to backup in case the LODModel is regenerated, this data is use to validate in the UI if the ddc must be rebuild
-	const FString BackupBuildStringID = SkeletalMesh->GetImportedModel()->LODModels[LODIndex].BuildStringID;
-
-	const FReferenceSkeleton& RefSkeleton = SkeletalMesh->GetRefSkeleton();
-
-	const bool bNaniteBuildEnabled = SkeletalMesh->IsNaniteEnabled();
-
-	FScopedSlowTask SlowTask(bNaniteBuildEnabled ? 7.01f : 6.01f, NSLOCTEXT("SkeltalMeshBuilder", "BuildingSkeletalMeshLOD", "Building skeletal mesh LOD"));
-	SlowTask.MakeDialog();
-
-	//Prevent any PostEdit change during the build
-	FScopedSkeletalMeshPostEditChange ScopedPostEditChange(SkeletalMesh, false, false);
-	// Unbind any existing clothing assets before we reimport the geometry
-	TArray<ClothingAssetUtils::FClothingAssetMeshBinding> ClothingBindings;
-	FLODUtilities::UnbindClothingAndBackup(SkeletalMesh, ClothingBindings, LODIndex);
-
-	int32 NumTexCoord = 1; //We need to send rendering at least one tex coord buffer
-
-	//This scope define where we can use the LODModel, after a reduction the LODModel must be requery since it is a new instance
+	void BuildLODModel()
 	{
-		FSkeletalMeshLODModel& BuildLODModel = SkeletalMesh->GetImportedModel()->LODModels[LODIndex];
-
-		//Load the imported data
-		const FMeshDescription& SkeletalMeshModel = *SkeletalMesh->GetMeshDescription(LODIndex);
-		FSkeletalMeshImportData SkeletalMeshImportData = FSkeletalMeshImportData::CreateFromMeshDescription(SkeletalMeshModel);
+		FSkeletalMeshLODModel& LODModel = GetLODModel();
 
 		TArray<FVector3f> LODPoints;
 		TArray<SkeletalMeshImportData::FMeshWedge> LODWedges;
 		TArray<SkeletalMeshImportData::FMeshFace> LODFaces;
 		TArray<SkeletalMeshImportData::FVertInfluence> LODInfluences;
 		TArray<int32> LODPointToRawMap;
-		SkeletalMeshImportData.CopyLODImportData(LODPoints, LODWedges, LODFaces, LODInfluences, LODPointToRawMap);
+		ImportData.CopyLODImportData(LODPoints, LODWedges, LODFaces, LODInfluences, LODPointToRawMap);
 
-		//Use the max because we need to have at least one texture coordinate
-		NumTexCoord = FMath::Max<int32>(NumTexCoord, SkeletalMeshImportData.NumTexCoords);
-		
 		// BaseLOD needs to make sure the source data fit with the skeletal mesh materials array before using meshutilities.BuildSkeletalMesh
-		FLODUtilities::AdjustImportDataFaceMaterialIndex(SkeletalMesh->GetMaterials(), SkeletalMeshImportData.Materials, LODFaces, LODIndex);
-
-		// Build the skeletal mesh using mesh utilities module
-		IMeshUtilities::MeshBuildOptions Options;
-		Options.FillOptions(LODInfo->BuildSettings);
-		Options.TargetPlatform = SkeletalMeshBuildParameters.TargetPlatform;
-		// Force the normals or tangent in case the data is missing
-		Options.bComputeNormals |= !SkeletalMeshImportData.bHasNormals;
-		Options.bComputeTangents |= !SkeletalMeshImportData.bHasTangents;
-		IMeshUtilities& MeshUtilities = FModuleManager::Get().LoadModuleChecked<IMeshUtilities>("MeshUtilities");
+		FLODUtilities::AdjustImportDataFaceMaterialIndex(SkeletalMesh->GetMaterials(), ImportData.Materials, LODFaces, LODIndex);
 
 		// Create skinning streams for NewModel.
-		SlowTask.EnterProgressFrame(1.0f);
-		MeshUtilities.BuildSkeletalMesh(
-			BuildLODModel, 
+		MeshUtilities->BuildSkeletalMesh(
+			LODModel, 
 			SkeletalMesh->GetPathName(),
-			RefSkeleton,
+			SkeletalMesh->GetRefSkeleton(),
 			LODInfluences,
 			LODWedges,
 			LODFaces,
@@ -256,7 +112,7 @@ bool FSkeletalMeshBuilder::Build(const FSkeletalMeshBuildParameters& SkeletalMes
 		);
 
 		// Set texture coordinate count on the new model.
-		BuildLODModel.NumTexCoords = NumTexCoord;
+		LODModel.NumTexCoords = GetNumTexCoords();
 
 		// Cache the vertex/triangle count in the InlineReductionCacheData so we can know if the LODModel need reduction or not.
 		TArray<FInlineReductionCacheData>& InlineReductionCacheDatas = SkeletalMesh->GetImportedModel()->InlineReductionCacheDatas;
@@ -266,57 +122,80 @@ bool FSkeletalMeshBuilder::Build(const FSkeletalMeshBuildParameters& SkeletalMes
 		}
 		if (ensure(InlineReductionCacheDatas.IsValidIndex(LODIndex)))
 		{
-			InlineReductionCacheDatas[LODIndex].SetCacheGeometryInfo(BuildLODModel);
+			InlineReductionCacheDatas[LODIndex].SetCacheGeometryInfo(LODModel);
 		}
 
 		// Re-Apply the user section changes, the UserSectionsData is map to original section and should match the built LODModel
-		BuildLODModel.SyncronizeUserSectionsDataArray();
+		LODModel.SyncronizeUserSectionsDataArray();
+	}
 
-		if (bNaniteBuildEnabled)
-		{
-			SlowTask.EnterProgressFrame(1.0f, NSLOCTEXT("SkeltalMeshBuilder", "BuildingNaniteData", "Building Nanite data..."));
+	bool BuildNanite()
+	{
+		FSkeletalMeshRenderData* SkeletalMeshRenderData = SkeletalMesh->GetResourceForRendering();
+		check(SkeletalMeshRenderData != nullptr);
 
-			FSkeletalMeshRenderData* SkeletalMeshRenderData = SkeletalMesh->GetResourceForRendering();
-			check(SkeletalMeshRenderData != nullptr);
+		ClearNaniteResources(SkeletalMeshRenderData->NaniteResourcesPtr);
+	
+		Nanite::IBuilderModule::FInputMeshData InputMeshData;
+		InitNaniteBuildInput(InputMeshData);
+
+		Nanite::FResources& NaniteResources = *SkeletalMeshRenderData->NaniteResourcesPtr.Get();
+		if (!NaniteBuilder->Build(
+			NaniteResources,
+			InputMeshData,
+			nullptr, // OutFallbackMeshData
+			nullptr, // OutRayTracingFallbackMeshData
+			nullptr, // RayTracingFallbackBuildSettings
+			NaniteSettings,
+			nullptr  // NaniteAssemblyData
+			)
+		)
+		{
+			return false;
+		}
 
-			ClearNaniteResources(SkeletalMeshRenderData->NaniteResourcesPtr);
+		return true;
+	}
 
-			Nanite::FResources& NaniteResources = *SkeletalMeshRenderData->NaniteResourcesPtr.Get();
+	void RestoreClothingFromBackup()
+	{
+		FLODUtilities::RestoreClothingFromBackup(SkeletalMesh, ClothingBindings, LODIndex);
 
-			bool bBuildSuccess = BuildNanite(
-				SkeletalMesh,
-				BuildLODModel,
-				SkeletalMeshModel,
-				NaniteResources
-			);
-		}
+		FSkeletalMeshLODModel& LODModel = GetLODModel();
+		LODModel.SyncronizeUserSectionsDataArray();
+		LODModel.NumTexCoords = GetNumTexCoords();
+		LODModel.BuildStringID = BackupBuildStringID;
+	}
 
-		// Re-apply the morph target
-		SlowTask.EnterProgressFrame(1.0f, NSLOCTEXT("SkeltalMeshBuilder", "RebuildMorphTarget", "Rebuilding morph targets..."));
-		if (SkeletalMeshImportData.MorphTargetNames.Num() > 0)
+	void BuildMorphTargets()
+	{
+		if (ImportData.MorphTargetNames.Num() > 0)
 		{
-			FLODUtilities::BuildMorphTargets(SkeletalMesh, SkeletalMeshModel, SkeletalMeshImportData, LODIndex, !Options.bComputeNormals, !Options.bComputeTangents, Options.bUseMikkTSpace, Options.OverlappingThresholds);
+			FLODUtilities::BuildMorphTargets(SkeletalMesh, *MeshDescription, ImportData, LODIndex, !Options.bComputeNormals, !Options.bComputeTangents, Options.bUseMikkTSpace, Options.OverlappingThresholds);
 		}
+	}
 
-		// Re-apply the alternate skinning it must be after the inline reduction
-		SlowTask.EnterProgressFrame(1.0f, NSLOCTEXT("SkeltalMeshBuilder", "RebuildAlternateSkinning", "Rebuilding alternate skinning..."));
-		const TArray<FSkinWeightProfileInfo> SkinProfiles = SkeletalMesh->GetSkinWeightProfiles();
-		for (int32 SkinProfileIndex = 0; SkinProfileIndex < SkinProfiles.Num(); ++SkinProfileIndex)
-		{
-			const FSkinWeightProfileInfo& ProfileInfo = SkinProfiles[SkinProfileIndex];
+	void UpdateAlternateSkinWeights()
+	{
+		for (const FSkinWeightProfileInfo& ProfileInfo : SkeletalMesh->GetSkinWeightProfiles())
+		{			
 			FLODUtilities::UpdateAlternateSkinWeights(SkeletalMesh, ProfileInfo.Name, LODIndex, Options);
 		}
+	}
 
-		// Copy vertex attribute definitions and their values from the import model.
+	void UpdateLODInfoVertexAttributes()
+	{
 		FLODUtilities::UpdateLODInfoVertexAttributes(SkeletalMesh, LODIndex, LODIndex, /*CopyAttributeValues*/true);
-		
+	}
+
+	void PerformReductions()
+	{
 		FSkeletalMeshUpdateContext UpdateContext;
 		UpdateContext.SkeletalMesh = SkeletalMesh;
 
 		// We are reduce ourself in this case we reduce ourself from the original data and return true.
 		if (SkeletalMesh->IsReductionActive(LODIndex))
 		{
-			SlowTask.EnterProgressFrame(1.0f, NSLOCTEXT("SkeltalMeshBuilder", "RegenerateLOD", "Regenerate LOD..."));
 			// Update the original reduction data since we just build a new LODModel.
 			if (LODInfo->ReductionSettings.BaseLOD == LODIndex && SkeletalMesh->HasMeshDescription(LODIndex))
 			{
@@ -325,7 +204,7 @@ bool FSkeletalMeshBuilder::Build(const FSkeletalMeshBuildParameters& SkeletalMes
 					SkeletalMesh->GetLODInfo(LODIndex)->SourceImportFilename = SkeletalMesh->GetAssetImportData()->GetFirstFilename();
 				}
 			}
-			FLODUtilities::SimplifySkeletalMeshLOD(UpdateContext, LODIndex, SkeletalMeshBuildParameters.TargetPlatform, false);
+			FLODUtilities::SimplifySkeletalMeshLOD(UpdateContext, LODIndex, TargetPlatform, false);
 		}
 		else
 		{
@@ -337,27 +216,206 @@ bool FSkeletalMeshBuilder::Build(const FSkeletalMeshBuildParameters& SkeletalMes
 				{
 					BonesToRemove.Add(BoneReference.BoneName);
 				}
-				MeshUtilities.RemoveBonesFromMesh(SkeletalMesh, LODIndex, &BonesToRemove);
+				MeshUtilities->RemoveBonesFromMesh(SkeletalMesh, LODIndex, &BonesToRemove);
+			}
+		}
+	}
+
+	void RegenerateDependentLODs()
+	{
+		if (bRegenDepLODs)
+		{
+			FLODUtilities::RegenerateDependentLODs(SkeletalMesh, LODIndex, TargetPlatform);
+		}
+	}
+
+private:
+	void Init()
+	{
+		check(SkeletalMesh->GetImportedModel());
+		check(SkeletalMesh->GetImportedModel()->LODModels.IsValidIndex(LODIndex));
+
+		LODInfo = SkeletalMesh->GetLODInfo(LODIndex);
+		check(LODInfo);
+
+		MeshDescription = SkeletalMesh->GetMeshDescription(LODIndex);
+		check(MeshDescription);
+
+		ImportData = FSkeletalMeshImportData::CreateFromMeshDescription(*MeshDescription);
+		
+		// Build the skeletal mesh using mesh utilities module
+		Options.FillOptions(LODInfo->BuildSettings);
+		Options.TargetPlatform = TargetPlatform;
+		
+		// Force the normals or tangent in case the data is missing
+		Options.bComputeNormals |= !ImportData.bHasNormals;
+		Options.bComputeTangents |= !ImportData.bHasTangents;
+
+		if (IsNaniteBuildEnabled())
+		{
+			NaniteSettings = SkeletalMesh->NaniteSettings;
+			NaniteSettings.KeepPercentTriangles = 1.0f;
+			NaniteSettings.TrimRelativeError = 0.0f;
+			NaniteSettings.FallbackPercentTriangles = 1.0f; // 100% - no reduction
+			NaniteSettings.FallbackRelativeError = 0.0f;
+		}
+	}
+
+	void InitNaniteBuildInput(Nanite::IBuilderModule::FInputMeshData& InputMeshData)
+	{
+		const FSkeletalMeshLODModel& LODModel = GetLODModel();
+
+		// Build new vertex buffers
+		InputMeshData.NumTexCoords = LODModel.NumTexCoords;
+
+		InputMeshData.MaterialIndices.SetNumUninitialized(LODModel.IndexBuffer.Num() / 3);
+
+		InputMeshData.Vertices.Position.SetNumUninitialized(LODModel.NumVertices);
+		InputMeshData.Vertices.TangentX.SetNumUninitialized(LODModel.NumVertices);
+		InputMeshData.Vertices.TangentY.SetNumUninitialized(LODModel.NumVertices);
+		InputMeshData.Vertices.TangentZ.SetNumUninitialized(LODModel.NumVertices);
+
+		InputMeshData.Vertices.UVs.SetNum(LODModel.NumTexCoords);
+		for (uint32 UVCoord = 0; UVCoord < LODModel.NumTexCoords; ++UVCoord)
+		{
+			InputMeshData.Vertices.UVs[UVCoord].SetNumUninitialized(LODModel.NumVertices);
+		}
+
+		// We can save memory by figuring out the max number of influences across all sections instead of allocating MAX_TOTAL_INFLUENCES
+		// Also check if any of the sections actually require 16bit, or if 8bit will suffice
+		bool b16BitSkinning = false;
+		InputMeshData.NumBoneInfluences = 0;
+		for (const FSkelMeshSection& Section : LODModel.Sections)
+		{
+			InputMeshData.NumBoneInfluences = FMath::Max< uint8 >(InputMeshData.NumBoneInfluences, uint32(Section.MaxBoneInfluences));
+			b16BitSkinning |= Section.Use16BitBoneIndex();
+		}
+
+		InputMeshData.Vertices.BoneIndices.SetNum(InputMeshData.NumBoneInfluences);
+		InputMeshData.Vertices.BoneWeights.SetNum(InputMeshData.NumBoneInfluences);
+		for (uint32 Influence = 0; Influence < InputMeshData.NumBoneInfluences; ++Influence)
+		{
+			InputMeshData.Vertices.BoneIndices[Influence].SetNumZeroed(LODModel.NumVertices);
+			InputMeshData.Vertices.BoneWeights[Influence].SetNumZeroed(LODModel.NumVertices);
+		}
+
+		// TODO: Nanite-Skinning
+		//InputMeshData.Vertices.Color.SetNumUninitialized(LODModel.NumVertices);
+
+		InputMeshData.TriangleIndices = LODModel.IndexBuffer;
+
+		uint32 CheckIndices = 0;
+		uint32 CheckVertices = 0;
+
+		for (int32 SectionIndex = 0; SectionIndex < LODModel.Sections.Num(); ++SectionIndex)
+		{
+			const FSkelMeshSection& Section = LODModel.Sections[SectionIndex];
+
+			check(CheckIndices  == Section.BaseIndex);
+			check(CheckVertices == Section.BaseVertexIndex);
+
+			for (int32 VertIndex = 0; VertIndex < Section.SoftVertices.Num(); ++VertIndex)
+			{
+				const FSoftSkinVertex& SoftVertex = Section.SoftVertices[VertIndex];
+
+				InputMeshData.Vertices.Position[Section.BaseVertexIndex + VertIndex] = SoftVertex.Position;
+				InputMeshData.Vertices.TangentX[Section.BaseVertexIndex + VertIndex] = SoftVertex.TangentX;
+				InputMeshData.Vertices.TangentY[Section.BaseVertexIndex + VertIndex] = SoftVertex.TangentY;
+				InputMeshData.Vertices.TangentZ[Section.BaseVertexIndex + VertIndex] = SoftVertex.TangentZ;
+
+				InputMeshData.VertexBounds += SoftVertex.Position;
+
+				for (uint32 UVCoord = 0; UVCoord < LODModel.NumTexCoords; ++UVCoord)
+				{
+					InputMeshData.Vertices.UVs[UVCoord][Section.BaseVertexIndex + VertIndex] = SoftVertex.UVs[UVCoord];
+				}
+
+				for (int32 Influence = 0; Influence < Section.MaxBoneInfluences; ++Influence)
+				{
+					InputMeshData.Vertices.BoneIndices[Influence][Section.BaseVertexIndex + VertIndex] = Section.BoneMap[SoftVertex.InfluenceBones[Influence]];
+					InputMeshData.Vertices.BoneWeights[Influence][Section.BaseVertexIndex + VertIndex] = SoftVertex.InfluenceWeights[Influence];
+				}
+
+				//InputMeshData.Vertices.Color[Section.BaseVertexIndex + VertIndex] = SoftVertex.Color;
+			}
+
+			for (uint32 MaterialIndex = 0; MaterialIndex < Section.NumTriangles; ++MaterialIndex)
+			{
+				InputMeshData.MaterialIndices[(CheckIndices / 3) + MaterialIndex] = Section.MaterialIndex;
 			}
+
+			CheckIndices += Section.NumTriangles * 3;
+			CheckVertices += Section.NumVertices;
 		}
+
+		check(CheckVertices == LODModel.NumVertices);
+		check(CheckIndices == LODModel.IndexBuffer.Num());
+
+		InputMeshData.TriangleCounts.Add(LODModel.IndexBuffer.Num() / 3);
 	}
 
-	FSkeletalMeshLODModel& LODModelAfterReduction = SkeletalMesh->GetImportedModel()->LODModels[LODIndex];
+	FString BackupBuildStringID;
+	TArray<ClothingAssetUtils::FClothingAssetMeshBinding> ClothingBindings;
+};
+
+} // namespace SkeletalMeshBuilderPrivate
+
+
+FSkeletalMeshBuilder::FSkeletalMeshBuilder()
+{
+}
+
+bool FSkeletalMeshBuilder::Build(const FSkeletalMeshBuildParameters& SkeletalMeshBuildParameters)
+{
+	TRACE_CPUPROFILER_EVENT_SCOPE(FSkeletalMeshBuilder::Build);
 	
-	// Re-apply the clothing using the UserSectionsData, this will ensure we remap correctly the cloth if the reduction has change the number of sections
-	SlowTask.EnterProgressFrame(1.0f, NSLOCTEXT("SkeltalMeshBuilder", "RebuildClothing", "Rebuilding clothing..."));
-	FLODUtilities::RestoreClothingFromBackup(SkeletalMesh, ClothingBindings, LODIndex);
+	SkeletalMeshBuilderPrivate::FContext Context(SkeletalMeshBuildParameters);
+	
+	const float TaskTotal = 6.01f + (Context.IsNaniteBuildEnabled() ? 1.f : 0.f);
+	FScopedSlowTask SlowTask(TaskTotal, NSLOCTEXT("SkeltalMeshBuilder", "BuildingSkeletalMeshLOD", "Building skeletal mesh LOD"));
+	SlowTask.MakeDialog();
 
-	LODModelAfterReduction.SyncronizeUserSectionsDataArray();
-	LODModelAfterReduction.NumTexCoords = NumTexCoord;
-	LODModelAfterReduction.BuildStringID = BackupBuildStringID;
+	//Prevent any PostEdit change during the build
+	FScopedSkeletalMeshPostEditChange ScopedPostEditChange(Context.SkeletalMesh, false, false);
 
-	SlowTask.EnterProgressFrame(1.0f, NSLOCTEXT("SkeltalMeshBuilder", "RegenerateDependentLODs", "Regenerate Dependent LODs..."));
-	if (SkeletalMeshBuildParameters.bRegenDepLODs)
+	// Unbind any existing clothing assets before we reimport the geometry
+	Context.UnbindClothingAndBackup();
+
+	SlowTask.EnterProgressFrame(1.0f);
+	Context.BuildLODModel();
+
+	if (Context.IsNaniteBuildEnabled())
 	{
-		// Regenerate dependent LODs
-		FLODUtilities::RegenerateDependentLODs(SkeletalMesh, LODIndex, SkeletalMeshBuildParameters.TargetPlatform);
+		TRACE_CPUPROFILER_EVENT_SCOPE(FSkeletalMeshBuilder::BuildNanite);
+
+		SlowTask.EnterProgressFrame(1.0f, NSLOCTEXT("SkeltalMeshBuilder", "BuildingNaniteData", "Building Nanite data..."));
+		bool bBuildSuccess = Context.BuildNanite();
+		if (!bBuildSuccess)
+		{
+			UE_LOG(LogSkeletalMesh, Error, TEXT("Failed to build Nanite for skeletal mesh. See previous line(s) for details."));
+		}
 	}
 
+	// Re-apply the morph target
+	SlowTask.EnterProgressFrame(1.0f, NSLOCTEXT("SkeltalMeshBuilder", "RebuildMorphTarget", "Rebuilding morph targets..."));
+	Context.BuildMorphTargets();
+
+	// Re-apply the alternate skinning it must be after the inline reduction
+	SlowTask.EnterProgressFrame(1.0f, NSLOCTEXT("SkeltalMeshBuilder", "RebuildAlternateSkinning", "Rebuilding alternate skinning..."));
+	Context.UpdateAlternateSkinWeights();
+
+	// Copy vertex attribute definitions and their values from the import model.
+	Context.UpdateLODInfoVertexAttributes();
+		
+	SlowTask.EnterProgressFrame(1.0f, NSLOCTEXT("SkeltalMeshBuilder", "RegenerateLOD", "Regenerate LOD..."));
+	Context.PerformReductions();
+	
+	// Re-apply the clothing using the UserSectionsData, this will ensure we remap correctly the cloth if the reduction has change the number of sections
+	SlowTask.EnterProgressFrame(1.0f, NSLOCTEXT("SkeltalMeshBuilder", "RebuildClothing", "Rebuilding clothing..."));
+	Context.RestoreClothingFromBackup();
+
+	SlowTask.EnterProgressFrame(1.0f, NSLOCTEXT("SkeltalMeshBuilder", "RegenerateDependentLODs", "Regenerate Dependent LODs..."));
+	Context.RegenerateDependentLODs();
+
 	return true;
 }
diff --git a/Engine/Source/Runtime/Core/Private/UObject/FortniteMainBranchObjectVersion.cpp b/Engine/Source/Runtime/Core/Private/UObject/FortniteMainBranchObjectVersion.cpp
--- a/Engine/Source/Runtime/Core/Private/UObject/FortniteMainBranchObjectVersion.cpp
+++ b/Engine/Source/Runtime/Core/Private/UObject/FortniteMainBranchObjectVersion.cpp
@@ -12,7 +12,7 @@ TMap<FGuid, FGuid> FFortniteMainBranchObjectVersion::GetSystemGuids()
 	SystemGuids.Add(DevGuids.NANITE_DERIVEDDATA_VER, FGuid("42C46B24E1586986880AFE62B0F575AB"));
 	SystemGuids.Add(DevGuids.NIAGARASHADERMAP_DERIVEDDATA_VER, FGuid("C417A50271B3427D98785388266E3B66"));
 	SystemGuids.Add(DevGuids.Niagara_LatestScriptCompileVersion, FGuid("393081597FA7234C92F00869DE262D73"));
-	SystemGuids.Add(DevGuids.SkeletalMeshDerivedDataVersion, FGuid("0829C9BAB43B4B25A63747FDA68E7197"));
+	SystemGuids.Add(DevGuids.SkeletalMeshDerivedDataVersion, FGuid("2F90180769024483BB7D90AACEB2798D"));
 	SystemGuids.Add(DevGuids.STATICMESH_DERIVEDDATA_VER, FGuid("656E338C78C34641A5D3BFAE303997AF"));
 	SystemGuids.Add(DevGuids.MaterialTranslationDDCVersion, FGuid("5DFD770145674B3AAFA544176CF7372B"));
 	return SystemGuids;
