Mesh builder interface updates:
- Refactor of skeletal mesh build functions to take render data as input to resolve issues where Nanite resources would not get properly built when cooking a Nanite skeletal mesh.
- Also updated some of the deprecated static mesh build methods to ensure any overrides by licensees get an error in the next release as aligns with our updated deprecation policy.

diff --git a/Engine/Plugins/ChaosClothAsset/Source/ChaosClothAssetEngine/Private/ChaosClothAsset/ClothAsset.cpp b/Engine/Plugins/ChaosClothAsset/Source/ChaosClothAssetEngine/Private/ChaosClothAsset/ClothAsset.cpp
--- a/Engine/Plugins/ChaosClothAsset/Source/ChaosClothAssetEngine/Private/ChaosClothAsset/ClothAsset.cpp
+++ b/Engine/Plugins/ChaosClothAsset/Source/ChaosClothAssetEngine/Private/ChaosClothAsset/ClothAsset.cpp
@@ -545,7 +545,7 @@ void UChaosClothAsset::BuildMeshModel()
 	for (int32 LodIndex = 0; LodIndex < NumLods; ++LodIndex)
 	{
 		GetImportedModel()->LODModels.Add(new FSkeletalMeshLODModel());
-		BuildLODModel(TargetPlatform, LodIndex);
+		FBuilder::BuildLod(GetImportedModel()->LODModels[LodIndex], *this, LodIndex, TargetPlatform);
 	}
 }
 #endif  // #if WITH_EDITOR
@@ -587,7 +587,7 @@ void UChaosClothAsset::CacheDerivedData(FSkinnedAssetCompilationContext* Context
 
 }
 
-void UChaosClothAsset::BuildLODModel(const ITargetPlatform* TargetPlatform, int32 LODIndex)
+void UChaosClothAsset::BuildLODModel(FSkeletalMeshRenderData& RenderData, const ITargetPlatform* TargetPlatform, int32 LODIndex)
 {
 	check(GetImportedModel() && GetImportedModel()->LODModels.IsValidIndex(LODIndex));
 	FBuilder::BuildLod(GetImportedModel()->LODModels[LODIndex], *this, LODIndex, TargetPlatform);
diff --git a/Engine/Plugins/ChaosClothAsset/Source/ChaosClothAssetEngine/Public/ChaosClothAsset/ClothAsset.h b/Engine/Plugins/ChaosClothAsset/Source/ChaosClothAssetEngine/Public/ChaosClothAsset/ClothAsset.h
--- a/Engine/Plugins/ChaosClothAsset/Source/ChaosClothAssetEngine/Public/ChaosClothAsset/ClothAsset.h
+++ b/Engine/Plugins/ChaosClothAsset/Source/ChaosClothAssetEngine/Public/ChaosClothAsset/ClothAsset.h
@@ -64,7 +64,7 @@ public:
 
 #if WITH_EDITOR
 	/* Build a LOD model for the targeted platform. */
-	virtual void BuildLODModel(const ITargetPlatform* TargetPlatform, int32 LODIndex) override;
+	virtual void BuildLODModel(FSkeletalMeshRenderData& RenderData, const ITargetPlatform* TargetPlatform, int32 LODIndex) override;
 	virtual FString BuildDerivedDataKey(const ITargetPlatform* TargetPlatform) override;
 	virtual bool IsInitialBuildDone() const override;
 #endif
diff --git a/Engine/Plugins/Experimental/ChaosOutfitAsset/Source/ChaosOutfitAssetEngine/Private/ChaosClothAsset/OutfitAsset.cpp b/Engine/Plugins/Experimental/ChaosOutfitAsset/Source/ChaosOutfitAssetEngine/Private/ChaosClothAsset/OutfitAsset.cpp
--- a/Engine/Plugins/Experimental/ChaosOutfitAsset/Source/ChaosOutfitAssetEngine/Private/ChaosClothAsset/OutfitAsset.cpp
+++ b/Engine/Plugins/Experimental/ChaosOutfitAsset/Source/ChaosOutfitAssetEngine/Private/ChaosClothAsset/OutfitAsset.cpp
@@ -61,7 +61,7 @@ const UPhysicsAsset* UChaosOutfitAsset::GetPhysicsAssetForModel(int32 ModelIndex
 }
 
 #if WITH_EDITOR
-void UChaosOutfitAsset::BuildLODModel(const ITargetPlatform* /*TargetPlatform*/, int32 /*LODIndex*/)
+void UChaosOutfitAsset::BuildLODModel(FSkeletalMeshRenderData& RenderData, const ITargetPlatform* /*TargetPlatform*/, int32 /*LODIndex*/)
 {
 }
 
diff --git a/Engine/Plugins/Experimental/ChaosOutfitAsset/Source/ChaosOutfitAssetEngine/Public/ChaosClothAsset/OutfitAsset.h b/Engine/Plugins/Experimental/ChaosOutfitAsset/Source/ChaosOutfitAssetEngine/Public/ChaosClothAsset/OutfitAsset.h
--- a/Engine/Plugins/Experimental/ChaosOutfitAsset/Source/ChaosOutfitAssetEngine/Public/ChaosClothAsset/OutfitAsset.h
+++ b/Engine/Plugins/Experimental/ChaosOutfitAsset/Source/ChaosOutfitAssetEngine/Public/ChaosClothAsset/OutfitAsset.h
@@ -51,7 +51,7 @@ private:
 	}
 	virtual void SetSkeleton(USkeleton* InSkeleton) override {}
 #if WITH_EDITOR
-	virtual void BuildLODModel(const ITargetPlatform* TargetPlatform, int32 LODIndex) override;
+	virtual void BuildLODModel(FSkeletalMeshRenderData& RenderData, const ITargetPlatform* TargetPlatform, int32 LODIndex) override;
 	virtual FString BuildDerivedDataKey(const ITargetPlatform* TargetPlatform) override;
 	virtual bool IsInitialBuildDone() const override;
 #endif
diff --git a/Engine/Source/Developer/MeshBuilder/Private/MeshBuilder.cpp b/Engine/Source/Developer/MeshBuilder/Private/MeshBuilder.cpp
--- a/Engine/Source/Developer/MeshBuilder/Private/MeshBuilder.cpp
+++ b/Engine/Source/Developer/MeshBuilder/Private/MeshBuilder.cpp
@@ -8,3 +8,25 @@ FMeshBuilder::FMeshBuilder()
 {
 
 }
+
+//
+// DEPRECATED METHODS & DEPENDENCIES
+//
+
+#include "Engine/StaticMesh.h"
+#include "Engine/SkeletalMesh.h"
+
+bool FMeshBuilder::Build(
+	FStaticMeshRenderData& OutRenderData,
+	UStaticMesh* StaticMesh,
+	const FStaticMeshLODGroup& LODGroup,
+	bool bAllowNanite)
+{
+	return Build(OutRenderData, FStaticMeshBuildParameters(StaticMesh, nullptr, LODGroup));
+}
+
+bool FMeshBuilder::Build(const FSkeletalMeshBuildParameters& SkeletalMeshBuildParameters)
+{
+	FSkeletalMeshRenderData* RenderData = SkeletalMeshBuildParameters.SkeletalMesh->GetResourceForRendering();
+	return Build(*RenderData, SkeletalMeshBuildParameters);
+}
\ No newline at end of file
diff --git a/Engine/Source/Developer/MeshBuilder/Private/MeshBuilderModule.cpp b/Engine/Source/Developer/MeshBuilder/Private/MeshBuilderModule.cpp
--- a/Engine/Source/Developer/MeshBuilder/Private/MeshBuilderModule.cpp
+++ b/Engine/Source/Developer/MeshBuilder/Private/MeshBuilderModule.cpp
@@ -28,15 +28,13 @@ public:
 
 	virtual bool BuildMesh(class FStaticMeshRenderData& OutRenderData, const FStaticMeshBuildParameters& BuildParameters) override;
 
-	virtual bool BuildMesh(FStaticMeshRenderData& OutRenderData, UObject* Mesh, const FStaticMeshLODGroup& LODGroup, bool bAllowNanite) override;
-
 	virtual bool BuildMeshVertexPositions(
 		UObject* StaticMesh,
 		TArray<uint32>& Indices,
 		TArray<FVector3f>& Vertices,
 		FStaticMeshSectionArray& Sections) override;
 
-	virtual bool BuildSkeletalMesh(const FSkeletalMeshBuildParameters& SkeletalMeshBuildParameters) override;
+	virtual bool BuildSkeletalMesh(FSkeletalMeshRenderData& OutRenderData, const FSkeletalMeshBuildParameters& SkeletalMeshBuildParameters) override;
 
 private:
 
@@ -50,16 +48,6 @@ bool FMeshBuilderModule::BuildMesh(FStaticMeshRenderData& OutRenderData, const F
 	return FStaticMeshBuilder().Build(OutRenderData, BuildParameters);
 }
 
-bool FMeshBuilderModule::BuildMesh(FStaticMeshRenderData& OutRenderData, class UObject* Mesh, const FStaticMeshLODGroup& LODGroup, bool bAllowNanite)
-{
-	UStaticMesh* StaticMesh = Cast<UStaticMesh>(Mesh);
-	if (StaticMesh != nullptr)
-	{
-		return BuildMesh(OutRenderData, FStaticMeshBuildParameters(StaticMesh, nullptr, LODGroup));
-	}
-	return false;
-}
-
 bool FMeshBuilderModule::BuildMeshVertexPositions(
 	UObject* Mesh,
 	TArray<uint32>& Indices,
@@ -75,8 +63,8 @@ bool FMeshBuilderModule::BuildMeshVertexPositions(
 	return false;
 }
 
-bool FMeshBuilderModule::BuildSkeletalMesh(const FSkeletalMeshBuildParameters& SkeletalMeshBuildParameters)
+bool FMeshBuilderModule::BuildSkeletalMesh(FSkeletalMeshRenderData& OutRenderData, const FSkeletalMeshBuildParameters& SkeletalMeshBuildParameters)
 {
 	//Call the skeletal mesh builder
-	return FSkeletalMeshBuilder().Build(SkeletalMeshBuildParameters);
+	return FSkeletalMeshBuilder().Build(OutRenderData, SkeletalMeshBuildParameters);
 }
\ No newline at end of file
diff --git a/Engine/Source/Developer/MeshBuilder/Private/SkeletalMeshBuilder.cpp b/Engine/Source/Developer/MeshBuilder/Private/SkeletalMeshBuilder.cpp
--- a/Engine/Source/Developer/MeshBuilder/Private/SkeletalMeshBuilder.cpp
+++ b/Engine/Source/Developer/MeshBuilder/Private/SkeletalMeshBuilder.cpp
@@ -129,17 +129,14 @@ public:
 		LODModel.SyncronizeUserSectionsDataArray();
 	}
 
-	bool BuildNanite()
+	bool BuildNanite(FSkeletalMeshRenderData& OutRenderData)
 	{
-		FSkeletalMeshRenderData* SkeletalMeshRenderData = SkeletalMesh->GetResourceForRendering();
-		check(SkeletalMeshRenderData != nullptr);
-
-		ClearNaniteResources(SkeletalMeshRenderData->NaniteResourcesPtr);
+		ClearNaniteResources(OutRenderData.NaniteResourcesPtr);
 	
 		Nanite::IBuilderModule::FInputMeshData InputMeshData;
 		InitNaniteBuildInput(InputMeshData);
 
-		Nanite::FResources& NaniteResources = *SkeletalMeshRenderData->NaniteResourcesPtr.Get();
+		Nanite::FResources& NaniteResources = *OutRenderData.NaniteResourcesPtr.Get();
 		if (!NaniteBuilder->Build(
 			NaniteResources,
 			InputMeshData,
@@ -365,7 +362,7 @@ FSkeletalMeshBuilder::FSkeletalMeshBuilder()
 {
 }
 
-bool FSkeletalMeshBuilder::Build(const FSkeletalMeshBuildParameters& SkeletalMeshBuildParameters)
+bool FSkeletalMeshBuilder::Build(FSkeletalMeshRenderData& OutRenderData, const FSkeletalMeshBuildParameters& SkeletalMeshBuildParameters)
 {
 	TRACE_CPUPROFILER_EVENT_SCOPE(FSkeletalMeshBuilder::Build);
 	
@@ -389,7 +386,7 @@ bool FSkeletalMeshBuilder::Build(const FSkeletalMeshBuildParameters& SkeletalMes
 		TRACE_CPUPROFILER_EVENT_SCOPE(FSkeletalMeshBuilder::BuildNanite);
 
 		SlowTask.EnterProgressFrame(1.0f, NSLOCTEXT("SkeltalMeshBuilder", "BuildingNaniteData", "Building Nanite data..."));
-		bool bBuildSuccess = Context.BuildNanite();
+		bool bBuildSuccess = Context.BuildNanite(OutRenderData);
 		if (!bBuildSuccess)
 		{
 			UE_LOG(LogSkeletalMesh, Error, TEXT("Failed to build Nanite for skeletal mesh. See previous line(s) for details."));
diff --git a/Engine/Source/Developer/MeshBuilder/Private/StaticMeshBuilder.cpp b/Engine/Source/Developer/MeshBuilder/Private/StaticMeshBuilder.cpp
--- a/Engine/Source/Developer/MeshBuilder/Private/StaticMeshBuilder.cpp
+++ b/Engine/Source/Developer/MeshBuilder/Private/StaticMeshBuilder.cpp
@@ -1325,15 +1325,6 @@ bool FStaticMeshBuilder::Build(FStaticMeshRenderData& StaticMeshRenderData, cons
 	return true;
 }
 
-bool FStaticMeshBuilder::Build(
-	FStaticMeshRenderData& OutRenderData,
-	UStaticMesh* StaticMesh,
-	const FStaticMeshLODGroup& LODGroup,
-	bool bAllowNanite)
-{
-	return Build(OutRenderData, FStaticMeshBuildParameters(StaticMesh, nullptr, LODGroup));
-}
-
 bool FStaticMeshBuilder::BuildMeshVertexPositions(
 	UStaticMesh* StaticMesh,
 	TArray<uint32>& BuiltIndices,
diff --git a/Engine/Source/Developer/MeshBuilder/Public/IMeshBuilderModule.h b/Engine/Source/Developer/MeshBuilder/Public/IMeshBuilderModule.h
--- a/Engine/Source/Developer/MeshBuilder/Public/IMeshBuilderModule.h
+++ b/Engine/Source/Developer/MeshBuilder/Public/IMeshBuilderModule.h
@@ -8,6 +8,10 @@
 #include "Interfaces/ITargetPlatform.h"
 #include "Interfaces/ITargetPlatformManagerModule.h"
 
+// TODO: Remove these dependencies when removing the deprecated functions below
+#include "Engine/StaticMesh.h"
+#include "Engine/SkeletalMesh.h"
+
 struct FStaticMeshSection;
 class FStaticMeshSectionArray;
 
@@ -31,11 +35,26 @@ public:
 	virtual bool BuildMesh(class FStaticMeshRenderData& OutRenderData, const struct FStaticMeshBuildParameters& BuildParameters) = 0;
 
 	UE_DEPRECATED(5.5, "Use FStaticMeshBuildParameters instead.")
-	virtual bool BuildMesh(class FStaticMeshRenderData& OutRenderData, class UObject* Mesh, const class FStaticMeshLODGroup& LODGroup, bool bAllowNanite) = 0;
+	virtual bool BuildMesh(class FStaticMeshRenderData& OutRenderData, class UObject* Mesh, const class FStaticMeshLODGroup& LODGroup, bool bAllowNanite) final
+	{
+		if (UStaticMesh* StaticMesh = Cast<UStaticMesh>(Mesh))
+		{
+			return BuildMesh(OutRenderData, FStaticMeshBuildParameters(StaticMesh, nullptr, LODGroup));
+		}
+		return false;
+	}
 
 	virtual bool BuildMeshVertexPositions(class UObject* StaticMesh, TArray<uint32>& Indices, TArray<FVector3f>& Vertices, FStaticMeshSectionArray& Sections) = 0;
 
-	virtual bool BuildSkeletalMesh(const struct FSkeletalMeshBuildParameters& SkeletalMeshBuildParameters) = 0;
+	UE_DEPRECATED(5.6, "Use the overload that takes an FSkeletalMeshRenderData instead.")
+	virtual bool BuildSkeletalMesh(const struct FSkeletalMeshBuildParameters& SkeletalMeshBuildParameters) final
+	{
+		FSkeletalMeshRenderData* RenderData = SkeletalMeshBuildParameters.SkeletalMesh->GetResourceForRendering();
+		check(RenderData != nullptr);
+		return BuildSkeletalMesh(*RenderData, SkeletalMeshBuildParameters);
+	}
+	
+	virtual bool BuildSkeletalMesh(class FSkeletalMeshRenderData& OutRenderData, const struct FSkeletalMeshBuildParameters& SkeletalMeshBuildParameters) = 0;
 
 	virtual void PostBuildSkeletalMesh(class FSkeletalMeshRenderData* SkeletalMeshRenderData, class USkinnedAsset* SkinnedAsset) { }
 };
diff --git a/Engine/Source/Developer/MeshBuilder/Public/MeshBuilder.h b/Engine/Source/Developer/MeshBuilder/Public/MeshBuilder.h
--- a/Engine/Source/Developer/MeshBuilder/Public/MeshBuilder.h
+++ b/Engine/Source/Developer/MeshBuilder/Public/MeshBuilder.h
@@ -12,6 +12,7 @@ class FStaticMeshSectionArray;
 class UStaticMesh;
 struct FStaticMeshBuildParameters;
 struct FSkeletalMeshBuildParameters;
+class FSkeletalMeshRenderData;
 struct FMeshBuildVertexData;
 
 /**
@@ -33,7 +34,7 @@ public:
 		FStaticMeshRenderData& OutRenderData,
 		UStaticMesh* StaticMesh,
 		const FStaticMeshLODGroup& LODGroup,
-		bool bAllowNanite) = 0;
+		bool bAllowNanite) final;
 
 	virtual bool BuildMeshVertexPositions(
 		UStaticMesh* StaticMesh,
@@ -44,7 +45,10 @@ public:
 	/**
 	 * Build function should be override and is the starting point for skeletal mesh builders
 	 */
-	virtual bool Build(const FSkeletalMeshBuildParameters& SkeletalMeshBuildParameters) = 0;
+	virtual bool Build(FSkeletalMeshRenderData& OutRenderData, const FSkeletalMeshBuildParameters& SkeletalMeshBuildParameters) = 0;
+
+	UE_DEPRECATED(5.6, "Use the overload that takes FSkeletalMeshRenderData instead")
+	virtual bool Build(const FSkeletalMeshBuildParameters& SkeletalMeshBuildParameters) final;
 
 private:
 
diff --git a/Engine/Source/Developer/MeshBuilder/Public/SkeletalMeshBuilder.h b/Engine/Source/Developer/MeshBuilder/Public/SkeletalMeshBuilder.h
--- a/Engine/Source/Developer/MeshBuilder/Public/SkeletalMeshBuilder.h
+++ b/Engine/Source/Developer/MeshBuilder/Public/SkeletalMeshBuilder.h
@@ -11,6 +11,7 @@ class UStaticMesh;
 class FStaticMeshRenderData;
 class FStaticMeshLODGroup;
 class USkeletalMesh;
+class FSkeletalMeshRenderData;
 
 class MESHBUILDER_API FSkeletalMeshBuilder : public FMeshBuilder
 {
@@ -25,14 +26,6 @@ public:
 		return false;
 	}
 
-	UE_DEPRECATED(5.5, "Use FStaticMeshBuildParameters instead.")
-	virtual bool Build(FStaticMeshRenderData& OutRenderData, UStaticMesh* StaticMesh, const FStaticMeshLODGroup& LODGroup, bool bAllowNanite) override
-	{
-		bool No_Support_For_StaticMesh_Build_In_FSkeletalMeshBuilder_Class = false;
-		check(No_Support_For_StaticMesh_Build_In_FSkeletalMeshBuilder_Class);
-		return false;
-	}
-
 	virtual bool BuildMeshVertexPositions(
 		UStaticMesh* StaticMesh,
 		TArray<uint32>& Indices,
@@ -44,7 +37,7 @@ public:
 		return false;
 	}
 	
-	virtual bool Build(const struct FSkeletalMeshBuildParameters& SkeletalMeshBuildParameters) override;
+	virtual bool Build(FSkeletalMeshRenderData& OutRenderData, const struct FSkeletalMeshBuildParameters& SkeletalMeshBuildParameters) override;
 
 	virtual ~FSkeletalMeshBuilder() {}
 
diff --git a/Engine/Source/Developer/MeshBuilder/Public/StaticMeshBuilder.h b/Engine/Source/Developer/MeshBuilder/Public/StaticMeshBuilder.h
--- a/Engine/Source/Developer/MeshBuilder/Public/StaticMeshBuilder.h
+++ b/Engine/Source/Developer/MeshBuilder/Public/StaticMeshBuilder.h
@@ -23,15 +23,8 @@ public:
 
 	virtual bool Build(FStaticMeshRenderData& OutRenderData, const FStaticMeshBuildParameters& BuildParameters) override;
 
-	UE_DEPRECATED(5.5, "Use FStaticMeshBuildParameters instead.")
-	virtual bool Build(
-		FStaticMeshRenderData& OutRenderData,
-		UStaticMesh* StaticMesh,
-		const FStaticMeshLODGroup& LODGroup,
-		bool bAllowNanite) override;
-
 	//No support for skeletal mesh build in this class
-	virtual bool Build(const struct FSkeletalMeshBuildParameters& SkeletalMeshBuildParameters) override
+	virtual bool Build(FSkeletalMeshRenderData& OutRenderData, const struct FSkeletalMeshBuildParameters& SkeletalMeshBuildParameters) override
 	{
 		bool No_Support_For_SkeletalMesh_Build_In_FStaticMeshBuilder_Class = false;
 		check(No_Support_For_SkeletalMesh_Build_In_FStaticMeshBuilder_Class);
diff --git a/Engine/Source/Editor/UnrealEd/Private/Fbx/FbxSkeletalMeshImport.cpp b/Engine/Source/Editor/UnrealEd/Private/Fbx/FbxSkeletalMeshImport.cpp
--- a/Engine/Source/Editor/UnrealEd/Private/Fbx/FbxSkeletalMeshImport.cpp
+++ b/Engine/Source/Editor/UnrealEd/Private/Fbx/FbxSkeletalMeshImport.cpp
@@ -1976,7 +1976,9 @@ USkeletalMesh* UnFbx::FFbxImporter::ImportSkeletalMesh(FImportSkeletalMeshArgs &
 			//We must build the LODModel so we can restore properly the mesh, but we do not have to regenerate LODs
 			const bool bRegenDepLODs = false;
 			FSkeletalMeshBuildParameters SkeletalMeshBuildParameters(SkeletalMesh, GetTargetPlatformManagerRef().GetRunningTargetPlatform(), ImportLODModelIndex, bRegenDepLODs);
-			bBuildSuccess = MeshBuilderModule.BuildSkeletalMesh(SkeletalMeshBuildParameters);
+			FSkeletalMeshRenderData* RenderData = SkeletalMesh->GetResourceForRendering();
+			check(RenderData != nullptr);
+			bBuildSuccess = MeshBuilderModule.BuildSkeletalMesh(*RenderData, SkeletalMeshBuildParameters);
 		}
 
 		if( !bBuildSuccess )
diff --git a/Engine/Source/Runtime/Core/Private/UObject/FortniteMainBranchObjectVersion.cpp b/Engine/Source/Runtime/Core/Private/UObject/FortniteMainBranchObjectVersion.cpp
--- a/Engine/Source/Runtime/Core/Private/UObject/FortniteMainBranchObjectVersion.cpp
+++ b/Engine/Source/Runtime/Core/Private/UObject/FortniteMainBranchObjectVersion.cpp
@@ -12,7 +12,7 @@ TMap<FGuid, FGuid> FFortniteMainBranchObjectVersion::GetSystemGuids()
 	SystemGuids.Add(DevGuids.NANITE_DERIVEDDATA_VER, FGuid("42C46B24E1586986880AFE62B0F575AB"));
 	SystemGuids.Add(DevGuids.NIAGARASHADERMAP_DERIVEDDATA_VER, FGuid("3AE847DB7C504839B140718982B04087"));
 	SystemGuids.Add(DevGuids.Niagara_LatestScriptCompileVersion, FGuid("670A2EBDEE1A4794BB77BBFF331EA7EF"));
-	SystemGuids.Add(DevGuids.SkeletalMeshDerivedDataVersion, FGuid("2F90180769024483BB7D90AACEB2798D"));
+	SystemGuids.Add(DevGuids.SkeletalMeshDerivedDataVersion, FGuid("5E93192FCF3647EEB7F9D4FFF1750050"));
 	SystemGuids.Add(DevGuids.STATICMESH_DERIVEDDATA_VER, FGuid("656E338C78C34641A5D3BFAE303997AF"));
 	SystemGuids.Add(DevGuids.MaterialTranslationDDCVersion, FGuid("D4A49317317F483C86516F96115A1CEC"));
 	return SystemGuids;
diff --git a/Engine/Source/Runtime/Engine/Classes/Engine/SkeletalMesh.h b/Engine/Source/Runtime/Engine/Classes/Engine/SkeletalMesh.h
--- a/Engine/Source/Runtime/Engine/Classes/Engine/SkeletalMesh.h
+++ b/Engine/Source/Runtime/Engine/Classes/Engine/SkeletalMesh.h
@@ -1266,8 +1266,8 @@ public:
 	/** Get the maximum number of optional LODs. Do not use MaxNumOptionalLODs directly. Call this method instead. USkinnedAsset Interface. */
 	ENGINE_API virtual int32 GetMaxNumOptionalLODs(const class ITargetPlatform* TargetPlatform) const override;
 
-	/* Build a LOD model before creating its render data. USkinnedAsset Interface. */
-	ENGINE_API virtual void BuildLODModel(const ITargetPlatform* TargetPlatform, int32 LODIndex) override;
+	/* Build a LOD model for the given render data. USkinnedAsset Interface. */
+	ENGINE_API virtual void BuildLODModel(FSkeletalMeshRenderData& RenderData, const ITargetPlatform* TargetPlatform, int32 LODIndex) override;
 #endif
 
 	UFUNCTION(BlueprintSetter)
diff --git a/Engine/Source/Runtime/Engine/Classes/Engine/SkinnedAsset.h b/Engine/Source/Runtime/Engine/Classes/Engine/SkinnedAsset.h
--- a/Engine/Source/Runtime/Engine/Classes/Engine/SkinnedAsset.h
+++ b/Engine/Source/Runtime/Engine/Classes/Engine/SkinnedAsset.h
@@ -259,7 +259,11 @@ public:
 	PURE_VIRTUAL(USkinnedAsset::IsInitialBuildDone, return false;);
 
 	/* Build a LOD model before creating its render data. */
-	virtual void BuildLODModel(const ITargetPlatform* TargetPlatform, int32 LODIndex) {}
+	UE_DEPRECATED(5.6, "Use or override the overload that takes FSkeletalMeshRenderData instead")
+	virtual void BuildLODModel(const ITargetPlatform* TargetPlatform, int32 LODIndex) final {}
+	
+	/* Build a LOD model for the given render data. */
+	virtual void BuildLODModel(FSkeletalMeshRenderData& RenderData, const ITargetPlatform* TargetPlatform, int32 LODIndex) {}
 
 	/** Get whether this mesh should use LOD streaming for the given platform. */
 	ENGINE_API virtual bool GetEnableLODStreaming(const ITargetPlatform* TargetPlatform) const
diff --git a/Engine/Source/Runtime/Engine/Private/SkeletalMesh.cpp b/Engine/Source/Runtime/Engine/Private/SkeletalMesh.cpp
--- a/Engine/Source/Runtime/Engine/Private/SkeletalMesh.cpp
+++ b/Engine/Source/Runtime/Engine/Private/SkeletalMesh.cpp
@@ -6084,7 +6084,7 @@ int32 USkeletalMesh::GetMaxNumOptionalLODs(const ITargetPlatform* TargetPlatform
 	}
 }
 
-void USkeletalMesh::BuildLODModel(const ITargetPlatform* TargetPlatform, int32 LODIndex)
+void USkeletalMesh::BuildLODModel(FSkeletalMeshRenderData& RenderData, const ITargetPlatform* TargetPlatform, int32 LODIndex)
 {
 	FSkeletalMeshModel* SkelMeshModel = GetImportedModel();
 	check(SkelMeshModel);
@@ -6102,7 +6102,7 @@ void USkeletalMesh::BuildLODModel(const ITargetPlatform* TargetPlatform, int32 L
 		LODInfoPtr->bHasBeenSimplified = false;
 		constexpr bool bRegenDepLODs = true;
 		const FSkeletalMeshBuildParameters BuildParameters(this, TargetPlatform, LODIndex, bRegenDepLODs);
-		MeshBuilderModule.BuildSkeletalMesh(BuildParameters);
+		MeshBuilderModule.BuildSkeletalMesh(RenderData, BuildParameters);
 	}
 	else
 	{
diff --git a/Engine/Source/Runtime/Engine/Private/SkeletalMeshRenderData.cpp b/Engine/Source/Runtime/Engine/Private/SkeletalMeshRenderData.cpp
--- a/Engine/Source/Runtime/Engine/Private/SkeletalMeshRenderData.cpp
+++ b/Engine/Source/Runtime/Engine/Private/SkeletalMeshRenderData.cpp
@@ -313,7 +313,7 @@ void FSkeletalMeshRenderData::Cache(const ITargetPlatform* TargetPlatform, USkin
 
 			for (int32 LODIndex = 0; LODIndex < SkelMeshModel->LODModels.Num(); LODIndex++)
 			{
-				Owner->BuildLODModel(TargetPlatform, LODIndex);
+				Owner->BuildLODModel(*this, TargetPlatform, LODIndex);
 
 				const FSkeletalMeshLODModel* LODModel = &(SkelMeshModel->LODModels[LODIndex]);
 				const FSkeletalMeshLODInfo* LODInfo = Owner->GetLODInfo(LODIndex);
