Temporarily force AllowFallbackGeneration to true for Nanite skel mesh (even if there are morphs). Some logic needs to be adjusted to support skipping it.

diff --git a/Engine/Source/Developer/MeshBuilder/Private/SkeletalMeshBuilder.cpp b/Engine/Source/Developer/MeshBuilder/Private/SkeletalMeshBuilder.cpp
--- a/Engine/Source/Developer/MeshBuilder/Private/SkeletalMeshBuilder.cpp
+++ b/Engine/Source/Developer/MeshBuilder/Private/SkeletalMeshBuilder.cpp
@@ -381,7 +381,7 @@ public:
 
 	inline bool AllowFallbackGeneration() const
 	{
-		return ImportData.MorphTargetNames.Num() == 0;
+		return true;// TODO: ImportData.MorphTargetNames.Num() == 0;
 	}
 
 private:
