Fix clothing asset binds on skeletal mesh assets during LOD model regen

diff --git a/Engine/Source/Developer/MeshBuilder/Private/SkeletalMeshBuilder.cpp b/Engine/Source/Developer/MeshBuilder/Private/SkeletalMeshBuilder.cpp
--- a/Engine/Source/Developer/MeshBuilder/Private/SkeletalMeshBuilder.cpp
+++ b/Engine/Source/Developer/MeshBuilder/Private/SkeletalMeshBuilder.cpp
@@ -583,9 +583,10 @@ bool FSkeletalMeshBuilder::Build(FSkeletalMeshRenderData& OutRenderData, const F
 	FMeshDescription* MeshDescriptionPtr = SkeletalMeshBuildParameters.SkeletalMesh->GetMeshDescription(SkeletalMeshBuildParameters.LODIndex);
 	FSkeletalMeshConstAttributes MeshAttributes(*MeshDescriptionPtr);
 
+	const bool bHasClothing = SkeletalMeshBuildParameters.SkeletalMesh->GetMeshClothingAssets().Num() > 0;
 	const bool bHasMorphTargets = MeshAttributes.GetMorphTargetNames().Num() > 0;
 	const bool bBuildNanite = SkeletalMeshBuildParameters.SkeletalMesh->IsNaniteEnabled() && SkeletalMeshBuildParameters.LODIndex == 0;
-	const bool bBuildNaniteFallback = bBuildNanite && !bHasMorphTargets;
+	const bool bBuildNaniteFallback = bBuildNanite && !bHasClothing && !bHasMorphTargets;
 
 	const float TaskTotal = 6.01f /*FinalizeContext*/ + (bBuildNanite ? 1.f : 0.f);
 	FScopedSlowTask SlowTask(TaskTotal, NSLOCTEXT("SkeletalMeshBuilder", "BuildingSkeletalMeshLOD", "Building skeletal mesh LOD"));
@@ -610,6 +611,9 @@ bool FSkeletalMeshBuilder::Build(FSkeletalMeshRenderData& OutRenderData, const F
 			SkeletalMeshBuilderPrivate::FContext NaniteContext(SkeletalMeshBuildParameters, &FallbackMeshDescription, bBuildNanite);
 			check(NaniteContext.IsNaniteBuildEnabled());
 
+			// Unbind any existing clothing assets before we reimport the geometry
+			NaniteContext.UnbindClothingAndBackup();
+
 			NaniteContext.BuildLODModel();
 
 			BuildNanite(SlowTask, NaniteContext, OutRenderData);
@@ -626,6 +630,9 @@ bool FSkeletalMeshBuilder::Build(FSkeletalMeshRenderData& OutRenderData, const F
 	{
 		SkeletalMeshBuilderPrivate::FContext Context(SkeletalMeshBuildParameters, MeshDescriptionPtr, bBuildNanite);
 
+		// Unbind any existing clothing assets before we reimport the geometry
+		Context.UnbindClothingAndBackup();
+
 		SlowTask.EnterProgressFrame(1.0f);
 		Context.BuildLODModel();
 
diff --git a/Engine/Source/Runtime/Core/Private/UObject/UE5MainStreamObjectVersion.cpp b/Engine/Source/Runtime/Core/Private/UObject/UE5MainStreamObjectVersion.cpp
--- a/Engine/Source/Runtime/Core/Private/UObject/UE5MainStreamObjectVersion.cpp
+++ b/Engine/Source/Runtime/Core/Private/UObject/UE5MainStreamObjectVersion.cpp
@@ -15,7 +15,7 @@ TMap<FGuid, FGuid> FUE5MainStreamObjectVersion::GetSystemGuids()
 	SystemGuids.Add(DevGuids.NANITE_DERIVEDDATA_VER, FGuid("49D2EF76967340EFB22F2F398E4DDDDD"));
 	SystemGuids.Add(DevGuids.NIAGARASHADERMAP_DERIVEDDATA_VER, FGuid("8A37C45D24F2423CBE5F8F371DE33575"));
 	SystemGuids.Add(DevGuids.Niagara_LatestScriptCompileVersion, FGuid("9C8812406D824DE29BBC6F303964BB7E"));
-	SystemGuids.Add(DevGuids.SkeletalMeshDerivedDataVersion, FGuid("D861782A8040482886D5C91F3F03DB8F"));
+	SystemGuids.Add(DevGuids.SkeletalMeshDerivedDataVersion, FGuid("F28EFD64577949D1983DCECC6BA3B9AF"));
 	SystemGuids.Add(DevGuids.STATICMESH_DERIVEDDATA_VER, FGuid("98F7D79A5811013825E75BBCC41ED3E9"));
 	SystemGuids.Add(DevGuids.POSESEARCHDB_DERIVEDDATA_VER, FGuid("4E595C2AC5E947D6BA9ABC874353E5BC"));
 	SystemGuids.Add(DevGuids.GROOM_BINDING_DERIVED_DATA_VERSION, FGuid("156678D4F9084D7CAE98FADC6AB93573"));
