Optimized skeletal mesh builder to avoid a mesh description clone when generating a Nanite fallback mesh, disabled cache optimization if the vertex count is large (like SM already does), and ensured that the build fallback test matches between high level and builder code.

diff --git a/Engine/Source/Developer/MeshBuilder/Private/SkeletalMeshBuilder.cpp b/Engine/Source/Developer/MeshBuilder/Private/SkeletalMeshBuilder.cpp
--- a/Engine/Source/Developer/MeshBuilder/Private/SkeletalMeshBuilder.cpp
+++ b/Engine/Source/Developer/MeshBuilder/Private/SkeletalMeshBuilder.cpp
@@ -177,25 +177,35 @@ namespace SkeletalMeshBuilderPrivate
 class FContext
 {
 public:
-	USkeletalMesh*						SkeletalMesh		= nullptr;
-	IMeshUtilities*						MeshUtilities		= nullptr;
-	Nanite::IBuilderModule*				NaniteBuilder		= nullptr;
-	const FSkeletalMeshLODInfo*			LODInfo				= nullptr;
-	FMeshDescription* 					MeshDescription		= nullptr;
-	const ITargetPlatform*				TargetPlatform		= nullptr;
+	USkeletalMesh*						SkeletalMesh			= nullptr;
+	IMeshUtilities*						MeshUtilities			= nullptr;
+	Nanite::IBuilderModule*				NaniteBuilder			= nullptr;
+	const FSkeletalMeshLODInfo*			LODInfo					= nullptr;
+	const FMeshDescription* 			SourceMeshDescription	= nullptr;
+	FMeshDescription*					FallbackMeshDescription	= nullptr;
+	const ITargetPlatform*				TargetPlatform			= nullptr;
 	FSkeletalMeshImportData 			ImportData;
 	IMeshUtilities::MeshBuildOptions	Options;
 	FMeshNaniteSettings 				NaniteSettings;
-	int32 								LODIndex			= INDEX_NONE;
-	bool								bRegenDepLODs		= false;
-
-	FContext(const FSkeletalMeshBuildParameters& InBuildParameters, FMeshDescription* InMeshDescription, bool bInBuildNanite)
+	int32 								LODIndex				= INDEX_NONE;
+	bool								bRegenDepLODs			= false;
+	bool								bBuildNaniteFallback	= false;
+
+	FContext(
+		const FSkeletalMeshBuildParameters& InBuildParameters,
+		const FMeshDescription* InSourceMeshDescription,
+		FMeshDescription* OutFallbackMeshDescription,
+		bool bInBuildNanite,
+		bool bInBuildNaniteFallback
+	)
 		: SkeletalMesh(InBuildParameters.SkeletalMesh)
 		, MeshUtilities(&FModuleManager::Get().LoadModuleChecked<IMeshUtilities>("MeshUtilities"))
-		, MeshDescription(InMeshDescription)
+		, SourceMeshDescription(InSourceMeshDescription)
+		, FallbackMeshDescription(OutFallbackMeshDescription)
 		, TargetPlatform(InBuildParameters.TargetPlatform)
 		, LODIndex(InBuildParameters.LODIndex)
 		, bRegenDepLODs(InBuildParameters.bRegenDepLODs)
+		, bBuildNaniteFallback(bInBuildNaniteFallback)
 	{
 		if (bInBuildNanite && SkeletalMesh->IsNaniteEnabled())
 		{
@@ -301,7 +311,8 @@ public:
 		// Fill out the mesh description for non-Nanite build/reduction
 		if (bGenerateFallback)
 		{
-			BuildNaniteFallbackMeshDescription(SkeletalMesh, FallbackMeshData, *MeshDescription);
+			check(FallbackMeshDescription);
+			BuildNaniteFallbackMeshDescription(SkeletalMesh, FallbackMeshData, *FallbackMeshDescription);
 		}
 
 		return true;
@@ -321,7 +332,7 @@ public:
 	{
 		if (ImportData.MorphTargetNames.Num() > 0)
 		{
-			FLODUtilities::BuildMorphTargets(SkeletalMesh, *MeshDescription, ImportData, LODIndex, !Options.bComputeNormals, !Options.bComputeTangents, Options.bUseMikkTSpace, Options.OverlappingThresholds);
+			FLODUtilities::BuildMorphTargets(SkeletalMesh, *SourceMeshDescription, ImportData, LODIndex, !Options.bComputeNormals, !Options.bComputeTangents, Options.bUseMikkTSpace, Options.OverlappingThresholds);
 		}
 	}
 
@@ -381,7 +392,7 @@ public:
 
 	inline bool AllowFallbackGeneration() const
 	{
-		return ImportData.MorphTargetNames.Num() == 0;
+		return IsNaniteBuildEnabled() && bBuildNaniteFallback && FallbackMeshDescription != nullptr;
 	}
 
 private:
@@ -393,8 +404,8 @@ private:
 		LODInfo = SkeletalMesh->GetLODInfo(LODIndex);
 		check(LODInfo);
 
-		check(MeshDescription);
-		ImportData = FSkeletalMeshImportData::CreateFromMeshDescription(*MeshDescription);
+		check(SourceMeshDescription);
+		ImportData = FSkeletalMeshImportData::CreateFromMeshDescription(*SourceMeshDescription);
 		
 		// Build the skeletal mesh using mesh utilities module
 		Options.FillOptions(LODInfo->BuildSettings);
@@ -408,7 +419,7 @@ private:
 		{
 			NaniteSettings = SkeletalMesh->NaniteSettings;
 			const bool bIsAssembly = false; // TODO: 
-			Nanite::CorrectFallbackSettings(NaniteSettings, MeshDescription->Triangles().Num(), bIsAssembly, /* bIsRayTracing */ false);
+			Nanite::CorrectFallbackSettings(NaniteSettings, SourceMeshDescription->Triangles().Num(), bIsAssembly, /* bIsRayTracing */ false);
 
 			if (Options.bComputeNormals && ImportData.bHasNormals)
 			{
@@ -423,6 +434,15 @@ private:
 			// Do not cache optimize the index buffer
 			Options.bCacheOptimize = false;
 		}
+
+		const int32 NumVertexInstances = SourceMeshDescription->VertexInstances().GetArraySize();
+		if (NumVertexInstances >= 100000 * 3)
+		{
+			// Just like static mesh, we disable cache optimization on very high poly 
+			// meshes because they are likely not for game rendering, or they are intended
+			// for rendering with Nanite.
+			Options.bCacheOptimize = false;
+		}
 	}
 
 	void InitNaniteBuildInput(Nanite::IBuilderModule::FInputMeshData& InputMeshData)
@@ -588,7 +608,11 @@ bool FSkeletalMeshBuilder::Build(FSkeletalMeshRenderData& OutRenderData, const F
 	const bool bBuildNanite = SkeletalMeshBuildParameters.SkeletalMesh->IsNaniteEnabled() && SkeletalMeshBuildParameters.LODIndex == 0;
 	const bool bBuildNaniteFallback = bBuildNanite && !bHasClothing && !bHasMorphTargets;
 
-	const float TaskTotal = 6.01f /*FinalizeContext*/ + (bBuildNanite ? 1.f : 0.f);
+	const float TaskTotal = 
+		5.01f /*FinalizeContext*/ + 
+		1.0f /* BuildLODModel */ +
+		(bBuildNanite ? 1.0f : 0.0f);
+
 	FScopedSlowTask SlowTask(TaskTotal, NSLOCTEXT("SkeletalMeshBuilder", "BuildingSkeletalMeshLOD", "Building skeletal mesh LOD"));
 	SlowTask.MakeDialog();
 
@@ -601,14 +625,13 @@ bool FSkeletalMeshBuilder::Build(FSkeletalMeshRenderData& OutRenderData, const F
 
 		FMeshDescription FallbackMeshDescription;
 		{
-			// Need to clone the mesh description as the fallback results from the Nanite build will modify the contents (and we do not want to commit the changes back).
-			if (!SkeletalMeshBuildParameters.SkeletalMesh->CloneMeshDescription(SkeletalMeshBuildParameters.LODIndex, FallbackMeshDescription))
-			{
-				UE_LOG(LogSkeletalMesh, Error, TEXT("Failed to clone skeletal mesh description. See previous line(s) for details."));
-				return false;
-			}
-
-			SkeletalMeshBuilderPrivate::FContext NaniteContext(SkeletalMeshBuildParameters, &FallbackMeshDescription, bBuildNanite);
+			SkeletalMeshBuilderPrivate::FContext NaniteContext(
+				SkeletalMeshBuildParameters,
+				MeshDescriptionPtr,
+				&FallbackMeshDescription,
+				bBuildNanite,
+				bBuildNaniteFallback
+			);
 			check(NaniteContext.IsNaniteBuildEnabled());
 
 			// Unbind any existing clothing assets before we reimport the geometry
@@ -619,7 +642,13 @@ bool FSkeletalMeshBuilder::Build(FSkeletalMeshRenderData& OutRenderData, const F
 			BuildNanite(SlowTask, NaniteContext, OutRenderData);
 		}
 
-		SkeletalMeshBuilderPrivate::FContext Context(SkeletalMeshBuildParameters, &FallbackMeshDescription, false /* bBuildNanite */);
+		SkeletalMeshBuilderPrivate::FContext Context(
+			SkeletalMeshBuildParameters,
+			&FallbackMeshDescription,
+			nullptr /* FallbackMeshDescription */,
+			false /* bBuildNanite */,
+			false /* bBuildNaniteFallback */
+		);
 
 		SlowTask.EnterProgressFrame(1.0f);
 		Context.BuildLODModel();
@@ -628,7 +657,13 @@ bool FSkeletalMeshBuilder::Build(FSkeletalMeshRenderData& OutRenderData, const F
 	}
 	else
 	{
-		SkeletalMeshBuilderPrivate::FContext Context(SkeletalMeshBuildParameters, MeshDescriptionPtr, bBuildNanite);
+		SkeletalMeshBuilderPrivate::FContext Context(
+			SkeletalMeshBuildParameters,
+			MeshDescriptionPtr,
+			nullptr /* FallbackMeshDescription */,
+			bBuildNanite,
+			false /* bBuildNaniteFallback */
+		);
 
 		// Unbind any existing clothing assets before we reimport the geometry
 		Context.UnbindClothingAndBackup();
diff --git a/Engine/Source/Runtime/Core/Private/UObject/UE5MainStreamObjectVersion.cpp b/Engine/Source/Runtime/Core/Private/UObject/UE5MainStreamObjectVersion.cpp
--- a/Engine/Source/Runtime/Core/Private/UObject/UE5MainStreamObjectVersion.cpp
+++ b/Engine/Source/Runtime/Core/Private/UObject/UE5MainStreamObjectVersion.cpp
@@ -15,7 +15,7 @@ TMap<FGuid, FGuid> FUE5MainStreamObjectVersion::GetSystemGuids()
 	SystemGuids.Add(DevGuids.NANITE_DERIVEDDATA_VER, FGuid("49D2EF76967340EFB22F2F398E4DDDDD"));
 	SystemGuids.Add(DevGuids.NIAGARASHADERMAP_DERIVEDDATA_VER, FGuid("8A37C45D24F2423CBE5F8F371DE33575"));
 	SystemGuids.Add(DevGuids.Niagara_LatestScriptCompileVersion, FGuid("9C8812406D824DE29BBC6F303964BB7E"));
-	SystemGuids.Add(DevGuids.SkeletalMeshDerivedDataVersion, FGuid("F28EFD64577949D1983DCECC6BA3B9AF"));
+	SystemGuids.Add(DevGuids.SkeletalMeshDerivedDataVersion, FGuid("B28EFD64577949B2083DCECC6BA3B9AC"));
 	SystemGuids.Add(DevGuids.STATICMESH_DERIVEDDATA_VER, FGuid("98F7D79A5811013825E75BBCC41ED3E9"));
 	SystemGuids.Add(DevGuids.POSESEARCHDB_DERIVEDDATA_VER, FGuid("4E595C2AC5E947D6BA9ABC874353E5BC"));
 	SystemGuids.Add(DevGuids.GROOM_BINDING_DERIVED_DATA_VERSION, FGuid("156678D4F9084D7CAE98FADC6AB93573"));
