[HWRT] Limit FoliageOverOcclusionBias to 0.9.
- setting it to 1.0 removes all triangles which leads to crashes in other code.
- if user actually intends to hide mesh from HWRT they can use "Support Ray Tracing" checkbox.

diff --git a/Engine/Source/Developer/MeshBuilder/Private/StaticMeshBuilder.cpp b/Engine/Source/Developer/MeshBuilder/Private/StaticMeshBuilder.cpp
--- a/Engine/Source/Developer/MeshBuilder/Private/StaticMeshBuilder.cpp
+++ b/Engine/Source/Developer/MeshBuilder/Private/StaticMeshBuilder.cpp
@@ -776,7 +776,9 @@ bool PrepareRayTracingFallbackBuild(
 
 			OutBuildContext.Settings.FallbackPercentTriangles = NaniteSettings.FallbackPercentTriangles;
 			OutBuildContext.Settings.FallbackRelativeError = NaniteSettings.FallbackRelativeError;
-			OutBuildContext.Settings.FoliageOverOcclusionBias = Settings.FoliageOverOcclusionBias;
+
+			// limit FoliageOverOcclusionBias since setting it to 1.0 removes all triangles and causes issues at runtime code paths.
+			OutBuildContext.Settings.FoliageOverOcclusionBias = FMath::Min(Settings.FoliageOverOcclusionBias, 0.9f);
 
 			return true;
 		}
diff --git a/Engine/Source/Editor/StaticMeshEditor/Private/StaticMeshEditorTools.cpp b/Engine/Source/Editor/StaticMeshEditor/Private/StaticMeshEditorTools.cpp
--- a/Engine/Source/Editor/StaticMeshEditor/Private/StaticMeshEditorTools.cpp
+++ b/Engine/Source/Editor/StaticMeshEditor/Private/StaticMeshEditorTools.cpp
@@ -5688,7 +5688,7 @@ void FRayTracingProxySettingsLayout::GenerateChildContent(IDetailChildrenBuilder
 			SNew(SSpinBox<float>)
 			.Font(IDetailLayoutBuilder::GetDetailFont())
 			.MinValue(0.0f)
-			.MaxValue(1.0f)
+			.MaxValue(0.9f)
 			.Value(this, &FRayTracingProxySettingsLayout::GetFoliageOverOcclusionBias)
 			.OnValueChanged(this, &FRayTracingProxySettingsLayout::OnFoliageOverOcclusionBiasChanged)
 			.OnValueCommitted(this, &FRayTracingProxySettingsLayout::OnFoliageOverOcclusionBiasCommitted)
