Added more verbose error message for BuildMeshVertexPositions for the cases where a NULL UObject or one that cannot be Cast to a StaicMesh is passed.

diff --git a/Engine/Source/Developer/MeshBuilder/Private/MeshBuilderModule.cpp b/Engine/Source/Developer/MeshBuilder/Private/MeshBuilderModule.cpp
--- a/Engine/Source/Developer/MeshBuilder/Private/MeshBuilderModule.cpp
+++ b/Engine/Source/Developer/MeshBuilder/Private/MeshBuilderModule.cpp
@@ -7,6 +7,9 @@
 #include "Engine/StaticMesh.h"
 #include "SkeletalMeshBuilder.h"
 #include "Engine/SkeletalMesh.h"
+#include "Logging/LogMacros.h"
+
+DEFINE_LOG_CATEGORY_STATIC(LogMeshBuilderModule, Display, All);
 
 class FMeshBuilderModule : public IMeshBuilderModule
 {
@@ -60,6 +63,17 @@ bool FMeshBuilderModule::BuildMeshVertexPositions(
 		//Call the static mesh builder
 		return FStaticMeshBuilder().BuildMeshVertexPositions(StaticMesh, Indices, Vertices, Sections);
 	}
+
+	// Add some log lines to help identify why BuildMeshVertexPositions has failed.
+	if ( Mesh==nullptr )
+	{
+		UE_LOG(LogMeshBuilderModule, Warning, TEXT("NULL Mesh UObject passed to BuildMeshVertexPositions"));
+	}	
+	else if (StaticMesh == nullptr)
+	{
+		UE_LOG(LogMeshBuilderModule, Warning, TEXT("Mesh UObject %s (%s) passed to BuildMeshVertexPositions could not be converted to a UStaticMesh"), *Mesh->GetFName().ToString(), *Mesh->GetPathName());
+	}
+
 	return false;
 }
 
