[BUGFIX] Legacy FBX Import - Skeletal Mesh Import Crash
- Skeletal Mesh Description could be nullptr when the mesh only has skeleton and but no polygons
- Added a nullptr check and early return

#rnx

diff --git a/Engine/Source/Developer/MeshBuilder/Private/SkeletalMeshBuilder.cpp b/Engine/Source/Developer/MeshBuilder/Private/SkeletalMeshBuilder.cpp
--- a/Engine/Source/Developer/MeshBuilder/Private/SkeletalMeshBuilder.cpp
+++ b/Engine/Source/Developer/MeshBuilder/Private/SkeletalMeshBuilder.cpp
@@ -633,6 +633,11 @@ bool FSkeletalMeshBuilder::Build(FSkeletalMeshRenderData& OutRenderData, const F
 	TRACE_CPUPROFILER_EVENT_SCOPE(FSkeletalMeshBuilder::Build);
 
 	FMeshDescription* MeshDescriptionPtr = SkeletalMeshBuildParameters.SkeletalMesh->GetMeshDescription(SkeletalMeshBuildParameters.LODIndex);
+	if (!MeshDescriptionPtr)
+	{
+		UE_LOG(LogSkeletalMesh, Error, TEXT("Failed to build skeletal mesh. Invalid Mesh Description at LODIndex %d"), SkeletalMeshBuildParameters.LODIndex);
+		return false;
+	}
 	FSkeletalMeshConstAttributes MeshAttributes(*MeshDescriptionPtr);
 
 	const bool bHasClothing = SkeletalMeshBuildParameters.SkeletalMesh->GetMeshClothingAssets().Num() > 0;
