Fixed crash in SkeletalMeshBuilder

#rnx

diff --git a/Engine/Source/Developer/MeshBuilder/Private/SkeletalMeshBuilder.cpp b/Engine/Source/Developer/MeshBuilder/Private/SkeletalMeshBuilder.cpp
--- a/Engine/Source/Developer/MeshBuilder/Private/SkeletalMeshBuilder.cpp
+++ b/Engine/Source/Developer/MeshBuilder/Private/SkeletalMeshBuilder.cpp
@@ -510,10 +510,15 @@ public:
 			// Update the original reduction data since we just build a new LODModel.
 			if (LODInfo->ReductionSettings.BaseLOD == LODIndex && SkeletalMesh->HasMeshDescription(LODIndex))
 			{
+#if WITH_EDITORONLY_DATA
 				if (LODIndex == 0)
 				{
-					SkeletalMesh->GetLODInfo(LODIndex)->SourceImportFilename = SkeletalMesh->GetAssetImportData()->GetFirstFilename();
+					if (const UAssetImportData* AssetImportData = SkeletalMesh->GetAssetImportData())
+					{
+						SkeletalMesh->GetLODInfo(LODIndex)->SourceImportFilename = AssetImportData->GetFirstFilename();
+					}
 				}
+#endif
 			}
 			FLODUtilities::SimplifySkeletalMeshLOD(UpdateContext, LODIndex, TargetPlatform, false);
 		}
