fix crash on bad material IDs generated via geometry script:
1. ensure() rather than crash if static mesh build encounters a triangle with polygon group ID not found in the mesh description's polygon group list
2. clamp material IDs in dynamic mesh -> mesh description conversion to be non-negative

#Synced-CL 50269774

diff --git a/Engine/Source/Developer/MeshBuilder/Private/StaticMeshBuilder.cpp b/Engine/Source/Developer/MeshBuilder/Private/StaticMeshBuilder.cpp
--- a/Engine/Source/Developer/MeshBuilder/Private/StaticMeshBuilder.cpp
+++ b/Engine/Source/Developer/MeshBuilder/Private/StaticMeshBuilder.cpp
@@ -1637,7 +1637,11 @@ void BuildVertexBuffer(
 	for (const FTriangleID TriangleID : MeshDescription.Triangles().GetElementIDs())
 	{
 		const FPolygonGroupID PolygonGroupID = MeshDescription.GetTrianglePolygonGroup(TriangleID);
-		const int32 SectionIndex = PolygonGroupToSectionIndex[PolygonGroupID];
+		
+		const int32* FoundSectionIndex = PolygonGroupToSectionIndex.Find(PolygonGroupID);
+		const int32 SectionIndex = 
+			ensureMsgf(FoundSectionIndex, TEXT("Static mesh builder found mesh description triangle with invalid polygon group ID not in the mesh description polygon groups list: %d"), PolygonGroupID.GetValue()) 
+			? *FoundSectionIndex : 0;
 		TArray<uint32>& SectionIndices = OutPerSectionIndices[SectionIndex];
 
 		TArrayView<const FVertexID> VertexIDs = MeshDescription.GetTriangleVertices(TriangleID);
diff --git a/Engine/Source/Runtime/MeshConversion/Private/DynamicMeshToMeshDescription.cpp b/Engine/Source/Runtime/MeshConversion/Private/DynamicMeshToMeshDescription.cpp
--- a/Engine/Source/Runtime/MeshConversion/Private/DynamicMeshToMeshDescription.cpp
+++ b/Engine/Source/Runtime/MeshConversion/Private/DynamicMeshToMeshDescription.cpp
@@ -773,6 +773,7 @@ void FDynamicMeshToMeshDescription::Convert_SharedInstances(const FDynamicMesh3*
 			{
 				MaterialID = MaterialIDToPolygroupIDMap[MaterialID];
 			}
+			MaterialID = FMath::Max(0, MaterialID);
 			UsePolygonGroupID = FPolygonGroupID(MaterialID);
 		}
 
@@ -1078,6 +1079,7 @@ void FDynamicMeshToMeshDescription::Convert_NoSharedInstances(const FDynamicMesh
 			{
 				MaterialID = MaterialIDToPolygroupIDMap[MaterialID];
 			}
+			MaterialID = FMath::Max(0, MaterialID);
 			UsePolygonGroupID = FPolygonGroupID(MaterialID);
 		}
 
