[HWRT] Modified how ray tracing proxy LOD1 is generated.
- removed per asset LOD1 Percent Triangle control.
- instead all assets use r.StaticMesh.RayTracingProxies.LOD1MaxNumOfTriangles.

Context:
- LOD1 should only ever be used for a couple of frames while LOD0 is streaming. Since LOD1 is always resident is can end up taking significant amount of pool space. With this change LOD1 is somewhat fixed size behaving similarly to Nanite root pages.

#Synced-CL 52158667

#Synced-CL 52158667

diff --git a/Engine/Source/Developer/MeshBuilder/Private/StaticMeshBuilder.cpp b/Engine/Source/Developer/MeshBuilder/Private/StaticMeshBuilder.cpp
--- a/Engine/Source/Developer/MeshBuilder/Private/StaticMeshBuilder.cpp
+++ b/Engine/Source/Developer/MeshBuilder/Private/StaticMeshBuilder.cpp
@@ -594,7 +594,7 @@ static void BuildNaniteFallbackMeshDescription(
 
 struct FRayTracingFallbackBuildContext
 {
-	TArray<float, TInlineAllocator<2>> PercentTriangles;
+	TArray<uint32, TInlineAllocator<2>> MaxNumOfTriangles;
 	TArray<FMeshDescription> MeshDescriptions;
 	Nanite::FRayTracingFallbackBuildSettings Settings;
 
@@ -761,6 +761,7 @@ bool PrepareRayTracingFallbackBuild(
 	if (StaticMesh->bSupportRayTracing && TargetPlatform->UsesRayTracing())
 	{
 		static const TConsoleVariableData<bool>* CVarRayTracingProxies = IConsoleManager::Get().FindTConsoleVariableDataBool(TEXT("r.StaticMesh.RayTracingProxies"));
+		static const TConsoleVariableData<int32>* CVarRayTracingProxiesLOD1MaxNumOfTriangles = IConsoleManager::Get().FindTConsoleVariableDataInt(TEXT("r.StaticMesh.RayTracingProxies.LOD1MaxNumOfTriangles"));
 
 		const FMeshRayTracingProxySettings& Settings = StaticMesh->GetRayTracingProxySettings();
 		const bool bNeedsRayTracingProxy =
@@ -769,10 +770,13 @@ bool PrepareRayTracingFallbackBuild(
 
 		if (bNeedsRayTracingProxy)
 		{
+			const uint32 LOD1MaxNumOfTriangles = CVarRayTracingProxiesLOD1MaxNumOfTriangles->GetValueOnAnyThread();
+
 			// LOD0 is generated by Nanite build and is already simplified so no need to reduce
-			OutBuildContext.PercentTriangles.Add(1.f);
-			OutBuildContext.PercentTriangles.Add(Settings.LOD1PercentTriangles);
-			OutBuildContext.MeshDescriptions.SetNum(OutBuildContext.PercentTriangles.Num());
+			OutBuildContext.MaxNumOfTriangles.Add(MAX_uint32);
+			OutBuildContext.MaxNumOfTriangles.Add(LOD1MaxNumOfTriangles);
+
+			OutBuildContext.MeshDescriptions.SetNum(OutBuildContext.MaxNumOfTriangles.Num());
 
 			FMeshNaniteSettings NaniteSettings;
 			NaniteSettings.FallbackTarget = Settings.FallbackTarget;
@@ -1244,11 +1248,11 @@ bool FStaticMeshBuilder::Build(FStaticMeshRenderData& StaticMeshRenderData, cons
 			FMeshDescriptionHelper MeshDescriptionHelper(&LODBuildSettings);
 
 			FMeshReductionSettings ReductionSettings;
-			ReductionSettings.PercentTriangles = FMath::Max(RayTracingFallbackBuildContext.PercentTriangles[LodIndex], .001);
+			ReductionSettings.MaxNumOfTriangles = RayTracingFallbackBuildContext.MaxNumOfTriangles[LodIndex];
 			ReductionSettings.TerminationCriterion = EStaticMeshReductionTerimationCriterion::Triangles;
 
-			const bool bUseReduction = ReductionSettings.PercentTriangles < 1.f;
 			TArray<FMeshDescription>& RayTracingMeshDescriptions = RayTracingFallbackBuildContext.MeshDescriptions;
+			const bool bUseReduction = ReductionSettings.PercentTriangles < 1.f || ReductionSettings.MaxNumOfTriangles < (uint32)RayTracingMeshDescriptions[BaseReduceLodIndex].Triangles().Num();
 
 			// Reduce LODs
 			if (bUseReduction)
diff --git a/Engine/Source/Editor/StaticMeshEditor/Private/StaticMeshEditorTools.cpp b/Engine/Source/Editor/StaticMeshEditor/Private/StaticMeshEditorTools.cpp
--- a/Engine/Source/Editor/StaticMeshEditor/Private/StaticMeshEditorTools.cpp
+++ b/Engine/Source/Editor/StaticMeshEditor/Private/StaticMeshEditorTools.cpp
@@ -5464,25 +5464,6 @@ void FRayTracingProxySettingsLayout::OnFallbackRelativeErrorCommitted(float NewV
 	OnFallbackRelativeErrorChanged(NewValue);
 }
 
-float FRayTracingProxySettingsLayout::GetLOD1PercentTriangles() const
-{
-	return RayTracingProxySettings.LOD1PercentTriangles * 100.0f; // Display fraction as percentage.
-}
-
-void FRayTracingProxySettingsLayout::OnLOD1PercentTrianglesChanged(float NewValue)
-{
-	RayTracingProxySettings.LOD1PercentTriangles = NewValue * 0.01f;
-}
-
-void FRayTracingProxySettingsLayout::OnLOD1PercentTrianglesCommitted(float NewValue, ETextCommit::Type TextCommitType)
-{
-	if (FEngineAnalytics::IsAvailable())
-	{
-		FEngineAnalytics::GetProvider().RecordEvent(TEXT("Editor.Usage.StaticMesh.RayTracingProxySettings"), TEXT("LOD1PercentTriangles"), FString::Printf(TEXT("%.1f"), NewValue));
-	}
-	OnLOD1PercentTrianglesChanged(NewValue);
-}
-
 float FRayTracingProxySettingsLayout::GetFoliageOverOcclusionBias() const
 {
 	return RayTracingProxySettings.FoliageOverOcclusionBias;
@@ -5644,30 +5625,6 @@ void FRayTracingProxySettingsLayout::GenerateChildContent(IDetailChildrenBuilder
 		})));
 	}
 
-	{
-		ChildrenBuilder.AddCustomRow(LOCTEXT("LOD1PercentTriangles", "LOD1 Percent Triangles"))
-		.RowTag("LOD1PercentTriangles")
-		.IsEnabled(RayTracingProxyEnabledAttr)
-		.NameContent()
-		[
-			SNew(STextBlock)
-			.Font(IDetailLayoutBuilder::GetDetailFont())
-			.Text(LOCTEXT("LOD1PercentTriangles", "LOD1 Percent Triangles"))
-			.ToolTipText(LOCTEXT("LOD1PercentTrianglesTooltip", "Percentage of triangles to keep in ray tracing proxy LOD 1."))
-		]
-		.ValueContent()
-		.VAlign(VAlign_Center)
-		[
-			SNew(SSpinBox<float>)
-			.Font(IDetailLayoutBuilder::GetDetailFont())
-			.MinValue(0.0f)
-			.MaxValue(100.0f)
-			.Value(this, &FRayTracingProxySettingsLayout::GetLOD1PercentTriangles)
-			.OnValueChanged(this, &FRayTracingProxySettingsLayout::OnLOD1PercentTrianglesChanged)
-			.OnValueCommitted(this, &FRayTracingProxySettingsLayout::OnLOD1PercentTrianglesCommitted)
-		];
-	}
-
 	{
 		ChildrenBuilder.AddCustomRow(LOCTEXT("FoliageOverOcclusionBias", "Foliage OverOcclusion Bias"))
 		.RowTag("FoliageOverOcclusionBias")
diff --git a/Engine/Source/Editor/StaticMeshEditor/Private/StaticMeshEditorTools.h b/Engine/Source/Editor/StaticMeshEditor/Private/StaticMeshEditorTools.h
--- a/Engine/Source/Editor/StaticMeshEditor/Private/StaticMeshEditorTools.h
+++ b/Engine/Source/Editor/StaticMeshEditor/Private/StaticMeshEditorTools.h
@@ -725,10 +725,6 @@ private:
 	void OnFallbackRelativeErrorChanged(float NewValue);
 	void OnFallbackRelativeErrorCommitted(float NewValue, ETextCommit::Type TextCommitType);
 
-	float GetLOD1PercentTriangles() const;
-	void OnLOD1PercentTrianglesChanged(float NewValue);
-	void OnLOD1PercentTrianglesCommitted(float NewValue, ETextCommit::Type TextCommitType);
-
 	float GetFoliageOverOcclusionBias() const;
 	void OnFoliageOverOcclusionBiasChanged(float NewValue);
 	void OnFoliageOverOcclusionBiasCommitted(float NewValue, ETextCommit::Type TextCommitType);
diff --git a/Engine/Source/Runtime/Engine/Classes/Engine/EngineTypes.h b/Engine/Source/Runtime/Engine/Classes/Engine/EngineTypes.h
--- a/Engine/Source/Runtime/Engine/Classes/Engine/EngineTypes.h
+++ b/Engine/Source/Runtime/Engine/Classes/Engine/EngineTypes.h
@@ -3550,9 +3550,6 @@ struct FMeshRayTracingProxySettings
 	/** Reduce until at least this amount of error is reached relative to size of the mesh */
 	UPROPERTY(BlueprintReadWrite, EditAnywhere, Category = RayTracing)
 	float FallbackRelativeError = 2.0f;
-
-	UPROPERTY(EditAnywhere, BlueprintReadWrite, Category = RayTracing, meta = (NoSpinbox = true, ClampMin = "0.0", ClampMax = "1.0", UIMin = "0.0", UIMax = "1.0"))
-	float LOD1PercentTriangles = 0.15f;
 	
 	/** A bias to reduce foliage over occlusion in Lumen GI. 0: no adjustment, 1: full strength. */
 	UPROPERTY(EditAnywhere, BlueprintReadWrite, Category = RayTracing, meta = (NoSpinbox = true, ClampMin = "0.0", ClampMax = "1.0", UIMin = "0.0", UIMax = "1.0"))
@@ -3570,7 +3567,6 @@ public:
 			Other.FallbackTarget == FallbackTarget &&
 			Other.FallbackPercentTriangles == FallbackPercentTriangles &&
 			Other.FallbackRelativeError == FallbackRelativeError &&
-			Other.LOD1PercentTriangles == LOD1PercentTriangles &&
 			Other.FoliageOverOcclusionBias == FoliageOverOcclusionBias;
 	}
 
diff --git a/Engine/Source/Runtime/Engine/Private/EngineTypes.cpp b/Engine/Source/Runtime/Engine/Private/EngineTypes.cpp
--- a/Engine/Source/Runtime/Engine/Private/EngineTypes.cpp
+++ b/Engine/Source/Runtime/Engine/Private/EngineTypes.cpp
@@ -227,7 +227,6 @@ void SerializeRayTracingProxySettingsForDDC(FArchive& Ar, FMeshRayTracingProxySe
 	Ar << Settings.FallbackTarget;
 	Ar << Settings.FallbackPercentTriangles;
 	Ar << Settings.FallbackRelativeError;
-	Ar << Settings.LOD1PercentTriangles;
 	Ar << Settings.FoliageOverOcclusionBias;
 }
 
diff --git a/Engine/Source/Runtime/Engine/Private/StaticMesh.cpp b/Engine/Source/Runtime/Engine/Private/StaticMesh.cpp
--- a/Engine/Source/Runtime/Engine/Private/StaticMesh.cpp
+++ b/Engine/Source/Runtime/Engine/Private/StaticMesh.cpp
@@ -181,6 +181,12 @@ static TAutoConsoleVariable<bool> CVarStaticMeshRayTracingProxies(
 	TEXT("Whether to use dedicated ray tracing LODs."),
 	ECVF_ReadOnly);
 
+static TAutoConsoleVariable<int32> CVarStaticMeshRayTracingProxiesLOD1MaxNumOfTriangles(
+	TEXT("r.StaticMesh.RayTracingProxies.LOD1MaxNumOfTriangles"),
+	64,
+	TEXT("Maximum number of triangles used for Ray Tracing Proxy LOD 1."),
+	ECVF_ReadOnly);
+
 static TAutoConsoleVariable<bool> CVarNaniteVoxelsUseSourceMeshUVDensity(
 	TEXT("r.Nanite.Voxels.UseSourceMeshUVDensity"),
 	1,
@@ -3889,6 +3895,9 @@ static FString BuildStaticMeshDerivedDataKeySuffix(const ITargetPlatform* Target
 			{
 				ByteToHex(SettingsAsBytes[ByteIndex], KeySuffix);
 			}
+
+			const int32 LOD1MaxNumOfTriangles = CVarStaticMeshRayTracingProxiesLOD1MaxNumOfTriangles.GetValueOnAnyThread();
+			KeySuffix += FString::Printf(TEXT("_LOD1(T%d)"), LOD1MaxNumOfTriangles);
 		}
 	}
 	else
